- LevelStreamingPersistence experimental plugin
- LevelStreamingPersistenceManager handles persistence of tagged object properties when streaming-out/in levels.
- Persistent object properties that are tagged as public can be accessed and modified even when the level is not loaded.

diff --git a/Engine/Plugins/Runtime/LevelStreamingPersistence/LevelStreamingPersistence.uplugin b/Engine/Plugins/Runtime/LevelStreamingPersistence/LevelStreamingPersistence.uplugin
new file mode 100644
--- /dev/null
+++ b/Engine/Plugins/Runtime/LevelStreamingPersistence/LevelStreamingPersistence.uplugin
@@ -0,0 +1,31 @@
+{
+	"FileVersion": 3,
+	"Version": 1,
+	"VersionName": "0.1",
+	"FriendlyName": "LevelStreamingPersistence",
+	"Description": "Experimental Level Streaming Persistence framework",
+	"Category": "",
+	"CreatedBy": "Epic Games, Inc.",
+	"CreatedByURL": "https://epicgames.com",
+	"DocsURL": "",
+	"MarketplaceURL": "",
+	"SupportURL": "",
+	"CanContainContent": false,
+	"IsBetaVersion": false,
+	"IsExperimentalVersion": true,
+	"Installed": false,
+	"EnabledByDefault": false,
+	"Modules": [
+		{
+			"Name": "LevelStreamingPersistence",
+			"Type": "Runtime",
+			"LoadingPhase": "PreLoadingScreen"
+		}
+	],
+	"Plugins": [
+		{
+			"Name": "StructUtils",
+			"Enabled": true
+		}
+	]
+}
\ No newline at end of file
diff --git a/Engine/Plugins/Runtime/LevelStreamingPersistence/Source/LevelStreamingPersistence/LevelStreamingPersistence.Build.cs b/Engine/Plugins/Runtime/LevelStreamingPersistence/Source/LevelStreamingPersistence/LevelStreamingPersistence.Build.cs
new file mode 100644
--- /dev/null
+++ b/Engine/Plugins/Runtime/LevelStreamingPersistence/Source/LevelStreamingPersistence/LevelStreamingPersistence.Build.cs
@@ -0,0 +1,27 @@
+// Copyright Epic Games, Inc. All Rights Reserved.
+
+namespace UnrealBuildTool.Rules
+{
+	public class LevelStreamingPersistence : ModuleRules
+	{
+		public LevelStreamingPersistence(ReadOnlyTargetRules Target) : base(Target)
+		{
+			PCHUsage = ModuleRules.PCHUsageMode.UseExplicitOrSharedPCHs;
+
+			PublicDependencyModuleNames.AddRange(
+				new string[] {
+					"Core",
+					"CoreUObject",
+					"Engine"
+				}
+			);
+
+			PrivateDependencyModuleNames.AddRange(
+				new string[] {
+					"StructUtils",
+					"PropertyPath"
+				}
+			);
+		}
+	}
+}
\ No newline at end of file
diff --git a/Engine/Plugins/Runtime/LevelStreamingPersistence/Source/LevelStreamingPersistence/Private/LevelStreamingPersistenceManager.cpp b/Engine/Plugins/Runtime/LevelStreamingPersistence/Source/LevelStreamingPersistence/Private/LevelStreamingPersistenceManager.cpp
new file mode 100644
--- /dev/null
+++ b/Engine/Plugins/Runtime/LevelStreamingPersistence/Source/LevelStreamingPersistence/Private/LevelStreamingPersistenceManager.cpp
@@ -0,0 +1,925 @@
+// Copyright Epic Games, Inc. All Rights Reserved.
+
+#include "LevelStreamingPersistenceManager.h"
+#include "LevelStreamingPersistenceModule.h"
+#include "LevelStreamingPersistenceSettings.h"
+#include "Serialization/ArchiveUObject.h"
+#include "Serialization/ObjectAndNameAsStringProxyArchive.h"
+#include "Serialization/MemoryReader.h"
+#include "Serialization/MemoryWriter.h"
+#include "Streaming/LevelStreamingDelegates.h"
+#include "WorldPartition/WorldPartition.h"
+#include "WorldPartition/WorldPartitionHelpers.h"
+#include "PropertyBag.h"
+#include "PropertyPathHelpers.h"
+#include "Engine/World.h"
+#include "Engine/Level.h"
+#include "Engine/Engine.h"
+#include "GameFramework/WorldSettings.h"
+
+DEFINE_LOG_CATEGORY_STATIC(LogLevelStreamingPersistenceManager, All, All);
+
+/*
+ * Achives used by ULevelStreamingPersistenceManager
+ */
+
+class FPersistentPropertiesArchive : public FObjectAndNameAsStringProxyArchive
+{
+public:
+	FPersistentPropertiesArchive(FArchive& InArchive, const FLevelStreamingPersistentPropertyArray& InCustomPropertyList)
+		: FObjectAndNameAsStringProxyArchive(InArchive, /*bLoadIfFindFails*/false)
+	{
+		check(InArchive.IsPersistent());
+		check(InArchive.IsFilterEditorOnly());
+		check(InArchive.ShouldSkipBulkData());
+		SetIsLoading(InArchive.IsLoading());
+		SetIsSaving(InArchive.IsSaving());
+		SetIsPersistent(true);
+		SetFilterEditorOnly(true);
+		ArShouldSkipBulkData = true;
+		ArUseCustomPropertyList = !InCustomPropertyList.IsEmpty();
+		ArCustomPropertyList = ArUseCustomPropertyList ? &InCustomPropertyList[0] : nullptr;
+	}
+
+	virtual FArchive& operator<<(FLazyObjectPtr& Value) override { return FArchiveUObject::SerializeLazyObjectPtr(*this, Value); }
+};
+
+class FPersistentPropertiesWriter : public FMemoryWriter
+{
+public:
+	FPersistentPropertiesWriter(TArray<uint8, TSizedDefaultAllocator<32>>& InBytes)
+		: FMemoryWriter(InBytes, true)
+	{
+		SetFilterEditorOnly(true);
+		ArShouldSkipBulkData = true;
+		SetIsTextFormat(false);
+		SetWantBinaryPropertySerialization(true);
+	}
+};
+
+class FPersistentPropertiesReader : public FMemoryReader
+{
+public:
+	FPersistentPropertiesReader(const TArray<uint8>& InBytes)
+		: FMemoryReader(InBytes, true)
+	{
+		SetFilterEditorOnly(true);
+		ArShouldSkipBulkData = true;
+		SetIsTextFormat(false);
+		SetWantBinaryPropertySerialization(true);
+	}
+};
+
+/*
+ * ULevelStreamingPersistenceManager implementation
+ */
+
+ULevelStreamingPersistenceManager::ULevelStreamingPersistenceManager(const FObjectInitializer& ObjectInitializer)
+	: PersistenceModule(&ILevelStreamingPersistenceModule::Get())
+{
+}
+
+bool ULevelStreamingPersistenceManager::DoesSupportWorldType(const EWorldType::Type InWorldType) const
+{
+	return InWorldType == EWorldType::PIE || InWorldType == EWorldType::Game;
+}
+
+bool ULevelStreamingPersistenceManager::ShouldCreateSubsystem(UObject* Outer) const
+{
+	if (!UWorldSubsystem::ShouldCreateSubsystem(Outer))
+	{
+		return false;
+	}
+
+	UWorld* World = Cast<UWorld>(Outer);
+	if (!World || !World->IsGameWorld() || (World->GetNetMode() == NM_Client))
+	{
+		return false;
+	}
+
+	UWorldPartition* WorldPartition = World->GetWorldPartition();
+	// @todo_ow: Make UWorldPartition::IsStreamingEnabled not dependant of WorldPartition's World member
+	if (WorldPartition && (!WorldPartition->bEnableStreaming || !World->GetWorldSettings()->SupportsWorldPartitionStreaming()))
+	{
+		return false;
+	}
+
+	return true;
+}
+
+void ULevelStreamingPersistenceManager::Initialize(FSubsystemCollectionBase& Collection)
+{
+	Super::Initialize(Collection);
+
+	FLevelStreamingDelegates::OnLevelBeginMakingInvisible.AddUObject(this, &ULevelStreamingPersistenceManager::OnLevelBeginMakingInvisible);
+	FLevelStreamingDelegates::OnLevelBeginMakingVisible.AddUObject(this, &ULevelStreamingPersistenceManager::OnLevelBeginMakingVisible);
+}
+
+void ULevelStreamingPersistenceManager::PostInitialize()
+{
+	PersistentPropertiesInfo.Initialize();
+}
+
+void ULevelStreamingPersistenceManager::Deinitialize()
+{
+	FLevelStreamingDelegates::OnLevelBeginMakingInvisible.RemoveAll(this);
+	FLevelStreamingDelegates::OnLevelBeginMakingVisible.RemoveAll(this);
+}
+
+bool ULevelStreamingPersistenceManager::TrySetPropertyValueFromString(const FString& InObjectPathName, const FName InPropertyName, const FString& InPropertyValue)
+{
+	const FString ObjectPathName = GetResolvedObjectPathName(InObjectPathName);
+	FLevelStreamingPersistentObjectPropertyBag* PropertyBag = GetPropertyBag(ObjectPathName);
+	if (PropertyBag && PropertyBag->SetPropertyValueFromString(InPropertyName, InPropertyValue))
+	{
+		// If object is loaded, copy property value to object
+		auto ObjectPropertyPair = GetObjectPropertyPair(ObjectPathName, InPropertyName);
+		CopyPropertyBagValueToObject(PropertyBag, ObjectPropertyPair.Key, ObjectPropertyPair.Value);
+		return true;
+	}
+	return false;
+}
+
+bool ULevelStreamingPersistenceManager::GetPropertyValueAsString(const FString& InObjectPathName, const FName InPropertyName, FString& OutPropertyValue)
+{
+	const FString ObjectPathName = GetResolvedObjectPathName(InObjectPathName);
+	if (FLevelStreamingPersistentObjectPropertyBag* PropertyBag = GetPropertyBag(ObjectPathName))
+	{
+		return PropertyBag->GetPropertyValueAsString(InPropertyName, OutPropertyValue);
+	}
+	return false;
+}
+
+void ULevelStreamingPersistenceManager::OnLevelBeginMakingInvisible(UWorld* World, const ULevelStreaming* InStreamingLevel, ULevel* InLoadedLevel)
+{
+	if (IsValid(InLoadedLevel) && (World == GetWorld()))
+	{
+		SaveLevelPersistentPropertyValues(InLoadedLevel);
+	}
+}
+
+void ULevelStreamingPersistenceManager::OnLevelBeginMakingVisible(UWorld* World, const ULevelStreaming* InStreamingLevel, ULevel* InLoadedLevel)
+{
+	if (IsValid(InLoadedLevel) && (World == GetWorld()))
+	{
+		RestoreLevelPersistentPropertyValues(InLoadedLevel);
+	}
+}
+
+bool ULevelStreamingPersistenceManager::SaveLevelPersistentPropertyValues(const ULevel* InLevel)
+{
+	if (!IsEnabled())
+	{
+		return false;
+	}
+
+	TRACE_CPUPROFILER_EVENT_SCOPE(ULevelStreamingPersistenceManager::SaveLevelPersistentPropertyValues);
+
+	const FString LevelPathName = InLevel->GetPathName();
+	FLevelStreamingPersistentPropertyValues& LevelProperties = LevelsPropertyValues.FindOrAdd(LevelPathName);
+
+	auto SavePrivateProperties = [this, &LevelProperties](const UObject* Object)
+	{
+		int32 SavedCount = 0;
+		check(IsValid(Object));
+		const UClass* ObjectClass = Object->GetClass();
+		if (PersistentPropertiesInfo.HasProperties(FLevelStreamingPersistentPropertiesInfo::PropertyType_Private, ObjectClass))
+		{
+			// Don't handle new objects (not part of original snapshot)
+			if (FLevelStreamingPersistentObjectPrivateProperties* ObjectPrivatePropertyValues = LevelProperties.ObjectsPrivatePropertyValues.Find(Object->GetPathName()))
+			{
+				ObjectPrivatePropertyValues->PayloadData.Reset();
+				ObjectPrivatePropertyValues->PersistentProperties.Reset();
+				if (DiffWithSnapshot(Object, ObjectPrivatePropertyValues->Snapshot, ObjectPrivatePropertyValues->PersistentProperties))
+				{
+					FLevelStreamingPersistentPropertyArray CustomPropertyList;
+					ObjectPrivatePropertyValues->BuildSerializedPropertyList(CustomPropertyList);
+					FPersistentPropertiesWriter WriterAr(ObjectPrivatePropertyValues->PayloadData);
+					FPersistentPropertiesArchive Ar(WriterAr, CustomPropertyList);
+					Object->SerializeScriptProperties(Ar);
+				}
+			}
+		}
+		return SavedCount;
+	};
+
+	auto SavePublicProperties = [this, &LevelProperties](const UObject* Object)
+	{
+		int32 SavedCount = 0;
+		check(IsValid(Object));
+		const UClass* ObjectClass = Object->GetClass();
+		// Save public properties
+		if (PersistentPropertiesInfo.HasProperties(FLevelStreamingPersistentPropertiesInfo::PropertyType_Public, ObjectClass))
+		{
+			TArray<FProperty*, TInlineAllocator<32>> SerializedObjectProperties;
+			PersistentPropertiesInfo.ForEachProperty(FLevelStreamingPersistentPropertiesInfo::PropertyType_Public, ObjectClass, [this, Object, &SerializedObjectProperties](FProperty* ObjectProperty)
+			{
+				if (!PersistenceModule->ShouldPersistProperty(Object, ObjectProperty))
+				{
+					return;
+				}
+				SerializedObjectProperties.Add(ObjectProperty);
+			});
+			if (!SerializedObjectProperties.IsEmpty())
+			{
+				FLevelStreamingPersistentObjectPublicProperties& ObjectPublicPropertyValues = LevelProperties.ObjectsPublicPropertyValues.FindOrAdd(Object->GetPathName());
+				if (ObjectPublicPropertyValues.PropertyBag.Initialize([this, ObjectClass]() { return PersistentPropertiesInfo.GetPropertyBagFromClass(FLevelStreamingPersistentPropertiesInfo::PropertyType_Public, ObjectClass); }))
+				{
+					ObjectPublicPropertyValues.PersistentProperties.Reset();
+					for (FProperty* ObjectProperty : SerializedObjectProperties)
+					{
+						if (const FProperty* BagProperty = ObjectPublicPropertyValues.PropertyBag.CopyPropertyValueFromObject(Object, ObjectProperty))
+						{
+							ObjectPublicPropertyValues.PersistentProperties.Add(BagProperty);
+							++SavedCount;
+						}
+					}
+				}
+			}
+		}
+		return SavedCount;
+	};
+
+	int32 SavedCount = 0;
+	ForEachObjectWithOuter(InLevel, [this, &SavePrivateProperties, &SavePublicProperties, &SavedCount](UObject* Object)
+	{	
+		SavedCount += SavePrivateProperties(Object);
+		SavedCount += SavePublicProperties(Object);
+	}, true, RF_NoFlags, EInternalObjectFlags::Garbage);
+	return SavedCount > 0;
+}
+
+bool ULevelStreamingPersistenceManager::RestoreLevelPersistentPropertyValues(ULevel* InLevel) const
+{
+	if (!IsEnabled())
+	{
+		return false;
+	}
+
+	TRACE_CPUPROFILER_EVENT_SCOPE(ULevelStreamingPersistenceManager::RestoreLevelPersistentPropertyValues);
+
+	auto SavePrivatePropertiesSnapshot = [this](UObject* Object, FLevelStreamingPersistentPropertyValues& LevelProperties)
+	{
+		int BuiltSnapshotCount = 0;
+		const UClass* ObjectClass = Object->GetClass();
+
+		if (PersistentPropertiesInfo.HasProperties(FLevelStreamingPersistentPropertiesInfo::PropertyType_Private, ObjectClass))
+		{
+			FLevelStreamingPersistentObjectPrivateProperties& ObjectPrivatePropertyValues = LevelProperties.ObjectsPrivatePropertyValues.FindOrAdd(Object->GetPathName());
+			if (!ObjectPrivatePropertyValues.Snapshot.IsValid())
+			{
+				if (BuildSnapshot(Object, ObjectPrivatePropertyValues.Snapshot))
+				{
+					++BuiltSnapshotCount;
+				}
+			}
+		}
+		return BuiltSnapshotCount;
+	};
+
+	int32 RestoredCount = 0;
+	int32 SavedCount = 0;
+
+	const FString LevelPathName = InLevel->GetPathName();
+	FLevelStreamingPersistentPropertyValues* LevelProperties = LevelsPropertyValues.Find(LevelPathName);
+	{
+		// The first time, create a snapshot of persistent properties
+		if (!LevelProperties)
+		{
+			TRACE_CPUPROFILER_EVENT_SCOPE(ULevelStreamingPersistenceManager::SavePrivatePropertiesSnapshot);
+
+			FLevelStreamingPersistentPropertyValues& NewLevelProperties = LevelsPropertyValues.FindOrAdd(LevelPathName);
+			ForEachObjectWithOuter(InLevel, [this, &NewLevelProperties, &SavedCount, &SavePrivatePropertiesSnapshot](UObject* Object)
+			{
+				SavedCount += SavePrivatePropertiesSnapshot(Object, NewLevelProperties);
+			}, true, RF_NoFlags, EInternalObjectFlags::Garbage);
+		}
+		// Else, restore persistent properties
+		else
+		{
+			for (const auto& [ObjectPath, ObjectPrivatePropertyValues] : LevelProperties->ObjectsPrivatePropertyValues)
+			{
+				if (UObject* Object = FindObject<UObject>(nullptr, *ObjectPath))
+				{
+					RestoredCount += RestorePrivateProperties(Object, ObjectPrivatePropertyValues);
+				}
+			}
+			for (const auto& [ObjectPath, ObjectPublicPropertyValues] : LevelProperties->ObjectsPublicPropertyValues)
+			{
+				if (UObject* Object = FindObject<UObject>(nullptr, *ObjectPath))
+				{
+					RestoredCount += RestorePublicProperties(Object, ObjectPublicPropertyValues);
+				}
+			}
+		}
+	}
+	
+	return (RestoredCount > 0) || (SavedCount > 0);
+}
+
+int32 ULevelStreamingPersistenceManager::RestorePrivateProperties(UObject* Object, const FLevelStreamingPersistentObjectPrivateProperties& PersistentObjectProperties) const
+{
+	int32 RestoredCount = 0;
+	check(IsValid(Object));
+
+	// Restore private properties
+	if (!PersistentObjectProperties.PayloadData.IsEmpty())
+	{
+		FLevelStreamingPersistentPropertyArray CustomPropertyList;
+		PersistentObjectProperties.BuildSerializedPropertyList(CustomPropertyList);
+		FPersistentPropertiesReader ReaderAr(PersistentObjectProperties.PayloadData);
+		FPersistentPropertiesArchive Ar(ReaderAr, CustomPropertyList);
+		Object->SerializeScriptProperties(Ar);
+		PersistentObjectProperties.ForEachPersistentProperty([this, Object](const FProperty* ObjectProperty)
+		{
+			PersistenceModule->PostRestorePersistedProperty(Object, ObjectProperty);
+		});
+		++RestoredCount;
+	}
+	return RestoredCount;
+};
+
+int32 ULevelStreamingPersistenceManager::RestorePublicProperties(UObject* Object, const FLevelStreamingPersistentObjectPublicProperties& ObjectsPublicPropertyValues) const
+{
+	int RestoredCount = 0;
+	check(IsValid(Object));
+	const UClass* ObjectClass = Object->GetClass();
+	// Restore public properties
+	if (ObjectsPublicPropertyValues.PropertyBag.IsValid())
+	{
+		ObjectsPublicPropertyValues.ForEachPersistentProperty([this, Object, ObjectClass, &ObjectsPublicPropertyValues, &RestoredCount](const FProperty* BagProperty)
+		{
+			FProperty* ObjectProperty = ObjectClass->FindPropertyByName(BagProperty->GetFName());
+			if (CopyPropertyBagValueToObject(&ObjectsPublicPropertyValues.PropertyBag, Object, ObjectProperty))
+			{
+				++RestoredCount;
+			}
+		});
+	}
+	return RestoredCount;
+};
+
+FLevelStreamingPersistentObjectPropertyBag* ULevelStreamingPersistenceManager::GetPropertyBag(const FString& InObjectPathName)
+{
+	FString LevelPathName;
+	FString ObjectShortPathName;
+	if (!SplitObjectPath(InObjectPathName, LevelPathName, ObjectShortPathName))
+	{
+		return nullptr;
+	}
+	FLevelStreamingPersistentPropertyValues* LevelProperties = LevelsPropertyValues.Find(LevelPathName);
+	FLevelStreamingPersistentObjectPublicProperties* ObjectPublicPropertyValues = LevelProperties ? LevelProperties->ObjectsPublicPropertyValues.Find(InObjectPathName) : nullptr;
+	return ObjectPublicPropertyValues ? &ObjectPublicPropertyValues->PropertyBag : nullptr;
+}
+
+const FLevelStreamingPersistentObjectPropertyBag* ULevelStreamingPersistenceManager::GetPropertyBag(const FString& InObjectPathName) const
+{
+	return const_cast<ULevelStreamingPersistenceManager*>(this)->GetPropertyBag(InObjectPathName);
+}
+
+bool ULevelStreamingPersistenceManager::CopyPropertyBagValueToObject(const FLevelStreamingPersistentObjectPropertyBag* InPropertyBag, UObject* InObject, FProperty* InObjectProperty) const
+{
+	if (InPropertyBag->CopyPropertyValueToObject(InObject, InObjectProperty))
+	{
+		PersistenceModule->PostRestorePersistedProperty(InObject, InObjectProperty);
+		return true;
+	}
+	return false;
+}
+
+TPair<UObject*, FProperty*> ULevelStreamingPersistenceManager::GetObjectPropertyPair(const FString& InObjectPathName, const FName InPropertyName) const
+{
+	UObject* Object = FindObject<UObject>(nullptr, *InObjectPathName);
+	FProperty* ObjectProperty = ::IsValid(Object) ? Object->GetClass()->FindPropertyByName(InPropertyName) : nullptr;
+	return TPair<UObject*, FProperty*>(Object, ObjectProperty);
+}
+
+const FString ULevelStreamingPersistenceManager::GetResolvedObjectPathName(const FString& InObjectPathName) const
+{
+	if (UWorldPartition* WorldPartition = GetWorld()->GetWorldPartition())
+	{
+		const FSoftObjectPath Path(InObjectPathName);
+		const FSoftObjectPath WorldPath(Path.GetAssetPath(), FString());
+		const UWorld* World = Cast<UWorld>(WorldPath.ResolveObject());
+		if (World && (World == GetWorld()))
+		{
+			FSoftObjectPath ResolvedPath;
+			if (FWorldPartitionHelpers::ConvertEditorPathToRuntimePath(FSoftObjectPath(InObjectPathName), ResolvedPath))
+			{
+				return ResolvedPath.ToString();
+			}
+		}
+	}
+	return InObjectPathName;
+}
+
+bool ULevelStreamingPersistenceManager::SplitObjectPath(const FString& InObjectPathName, FString& OutLevelPathName, FString& OutShortObjectPathName) const
+{
+	FSoftObjectPath ObjectPath(InObjectPathName);
+	if (ObjectPath.GetSubPathString().StartsWith(TEXT("PersistentLevel.")))
+	{
+		TStringBuilder<FName::StringBufferSize> Builder;
+		Builder << ObjectPath.GetAssetPath() << SUBOBJECT_DELIMITER_CHAR << TEXT("PersistentLevel");
+		OutLevelPathName = Builder.ToString();
+		OutShortObjectPathName = ObjectPath.GetSubPathString().RightChop(16);
+		return true;
+	}
+	return false;
+}
+
+
+bool ULevelStreamingPersistenceManager::BuildSnapshot(const UObject* InObject, FLevelStreamingPersistentObjectPropertyBag& OutSnapshot) const
+{
+	int SavedCount = 0;
+	const UClass* ObjectClass = InObject->GetClass();
+	if (PersistentPropertiesInfo.HasProperties(FLevelStreamingPersistentPropertiesInfo::PropertyType_Private, ObjectClass))
+	{
+		if (OutSnapshot.Initialize([this, ObjectClass]() { return PersistentPropertiesInfo.GetPropertyBagFromClass(FLevelStreamingPersistentPropertiesInfo::PropertyType_Private, ObjectClass); }))
+		{
+			PersistentPropertiesInfo.ForEachProperty(FLevelStreamingPersistentPropertiesInfo::PropertyType_Private, ObjectClass, [this, InObject, &SavedCount, &OutSnapshot](FProperty* ObjectProperty)
+			{
+				if (OutSnapshot.CopyPropertyValueFromObject(InObject, ObjectProperty))
+				{
+					++SavedCount;
+				}
+			});
+		}
+		check(OutSnapshot.IsValid());
+	}
+	return SavedCount > 0;
+}
+
+bool ULevelStreamingPersistenceManager::DiffWithSnapshot(const UObject* InObject, const FLevelStreamingPersistentObjectPropertyBag& InSnapshot, TArray<const FProperty*>& OutChangedProperties) const
+{
+	if (ensure(InSnapshot.IsValid()))
+	{
+		const UClass* ObjectClass = InObject->GetClass();
+
+		// Find changed properties
+		PersistentPropertiesInfo.ForEachProperty(FLevelStreamingPersistentPropertiesInfo::PropertyType_Private, ObjectClass, [this, InObject, &InSnapshot, &OutChangedProperties](FProperty* ObjectProperty)
+		{
+			if (!PersistenceModule->ShouldPersistProperty(InObject, ObjectProperty))
+			{
+				return;
+			}
+			bool bIsIdentical = true;
+			if (InSnapshot.ComparePropertyValueWithObject(InObject, ObjectProperty, bIsIdentical) && !bIsIdentical)
+			{
+				OutChangedProperties.Add(ObjectProperty);
+			}
+		});
+
+		return OutChangedProperties.Num() > 0;
+	}
+
+	return false;
+}
+
+void ULevelStreamingPersistenceManager::DumpContent() const
+{
+	UE_LOG(LogLevelStreamingPersistenceManager, Log, TEXT("World Persistence Content for %s"), *GetWorld()->GetName());
+
+	TMap<const UClass*, UObject*> TemporaryObjects;
+	auto GetTemporaryObjectForClass = [&TemporaryObjects](const UClass* InClass)
+	{
+		UObject*& Object = TemporaryObjects.FindOrAdd(InClass);
+		if (!Object)
+		{
+			Object = NewObject<UObject>(GetTransientPackage(), InClass);
+		}
+		return Object;
+	};
+
+	for (auto& [LevelPathName, LevelPropertyValues] : LevelsPropertyValues)
+	{
+		if (LevelPropertyValues.ObjectsPrivatePropertyValues.Num())
+		{
+			UE_LOG(LogLevelStreamingPersistenceManager, Log, TEXT("[+] Level Private Properties of %s"), *LevelPathName);
+			for (const auto& [ObjectPathName, ObjectPropertyValues] : LevelPropertyValues.ObjectsPrivatePropertyValues)
+			{
+				if (!ObjectPropertyValues.PayloadData.IsEmpty())
+				{
+					if (const UClass* Class = PersistentPropertiesInfo.GetClassFromPropertyBag(FLevelStreamingPersistentPropertiesInfo::PropertyType_Private, ObjectPropertyValues.Snapshot.GetPropertyBagStruct()))
+					{
+						if (UObject* TempObject = GetTemporaryObjectForClass(Class))
+						{
+							FString ObjectLevelPathName;
+							FString ObjectShortPathName;
+							const bool bUseShortName = SplitObjectPath(ObjectPathName, ObjectLevelPathName, ObjectShortPathName);
+
+							UE_LOG(LogLevelStreamingPersistenceManager, Log, TEXT("  [+] Object Private Properties of %s"), bUseShortName ? *ObjectShortPathName : *ObjectPathName);
+
+							if (RestorePrivateProperties(TempObject, ObjectPropertyValues))
+							{
+								for (const FProperty* Property : ObjectPropertyValues.PersistentProperties)
+								{
+									FString PropertyValue;
+									if (PropertyPathHelpers::GetPropertyValueAsString(TempObject, Property->GetName(), PropertyValue))
+									{
+										UE_LOG(LogLevelStreamingPersistenceManager, Log, TEXT("   - Property[%s] = %s"), *Property->GetName(), *PropertyValue);
+									}
+								}
+							}
+						}
+					}
+				}
+			}
+		}
+
+		if (LevelPropertyValues.ObjectsPublicPropertyValues.Num())
+		{
+			UE_LOG(LogLevelStreamingPersistenceManager, Log, TEXT("[+] Level Public Properties of %s"), *LevelPathName);
+			for (auto& [ObjectPathName, ObjectPublicPropertyValues] : LevelPropertyValues.ObjectsPublicPropertyValues)
+			{
+				if (ObjectPublicPropertyValues.PropertyBag.IsValid())
+				{
+					FString ObjectLevelPathName;
+					FString ObjectShortPathName;
+					const bool bUseShortName = SplitObjectPath(ObjectPathName, ObjectLevelPathName, ObjectShortPathName);
+
+					UE_LOG(LogLevelStreamingPersistenceManager, Log, TEXT("  [+] Object Public Properties of %s"), bUseShortName ? *ObjectShortPathName : *ObjectPathName);
+					ObjectPublicPropertyValues.PropertyBag.DumpContent([](const FProperty* Property, const FString& PropertyValue)
+					{
+						UE_LOG(LogLevelStreamingPersistenceManager, Log, TEXT("   - Property[%s] = %s"), *Property->GetName(), *PropertyValue);
+					}, &ObjectPublicPropertyValues.PersistentProperties);
+				}
+			}
+		}
+	}
+}
+
+/*
+ * FLevelStreamingPersistentObjectPrivateProperties implementation
+ */
+
+void FLevelStreamingPersistentObjectPrivateProperties::ForEachPersistentProperty(TFunctionRef<void(const FProperty*)> Func) const
+{
+	for (const FProperty* Property : PersistentProperties)
+	{
+		Func(Property);
+	}
+}
+
+void FLevelStreamingPersistentObjectPrivateProperties::BuildSerializedPropertyList(FLevelStreamingPersistentPropertyArray& OutCustomPropertyList) const
+{
+	ForEachPersistentProperty([&OutCustomPropertyList](const FProperty* Property)
+	{
+		OutCustomPropertyList.Emplace(const_cast<FProperty*>(Property));
+	});
+
+	// Link changed properties
+	if (!OutCustomPropertyList.IsEmpty())
+	{
+		for (int i = 0; i < OutCustomPropertyList.Num() - 1; ++i)
+		{
+			OutCustomPropertyList[i].PropertyListNext = &OutCustomPropertyList[i + 1];
+		}
+	}
+}
+
+/*
+ * FLevelStreamingPersistentObjectPublicProperties implementation
+ */
+
+void FLevelStreamingPersistentObjectPublicProperties::ForEachPersistentProperty(TFunctionRef<void(const FProperty*)> Func) const
+{
+	for (const FProperty* Property : PersistentProperties)
+	{
+		Func(Property);
+	}
+}
+
+/*
+ * FLevelStreamingPersistentPropertiesInfo implementation
+ */
+
+void FLevelStreamingPersistentPropertiesInfo::Initialize()
+{
+	auto FillClassesProperties = [this](const FString& InPropertyPath, bool bIsPublic)
+	{
+		TMap<TWeakObjectPtr<const UClass>, TSet<FProperty*>>& OutClassesProperties = ClassesProperties[bIsPublic ? PropertyType_Public : PropertyType_Private];
+		if (FProperty* Property = TFieldPath<FProperty>(*InPropertyPath).Get())
+		{
+			if (const UClass* Class = Property->GetOwnerClass())
+			{
+				OutClassesProperties.FindOrAdd(Class).Add(Property);
+			}
+		}
+	};
+
+	auto CreateClassDefaultPropertyBag = [this](EPropertyType InAccessSpecifier)
+	{
+		TMap<TWeakObjectPtr<const UClass>, FInstancedPropertyBag>& ClassesDefaults = ObjectClassToPropertyBag[InAccessSpecifier];
+		for (auto& [Class, ClassProperties] : ClassesProperties[InAccessSpecifier])
+		{
+			check(Class.IsValid());
+			check(!ClassesDefaults.Contains(Class));
+
+			TArray<FPropertyBagPropertyDesc> Descs;
+			ForEachProperty(InAccessSpecifier, Class.Get(), [&Descs](FProperty* Property)
+			{
+				Descs.Emplace(Property->GetFName(), Property);
+			});
+
+			FInstancedPropertyBag& PropertyBag = ClassesDefaults.FindOrAdd(Class);
+			PropertyBag.AddProperties(Descs);
+			check(PropertyBag.GetPropertyBagStruct());
+			PropertyBagToObjectClass[InAccessSpecifier].Add(PropertyBag.GetPropertyBagStruct(), Class);
+		}
+	};
+
+	for (const FLevelStreamingPersistentProperty& PersistentProperty : GetDefault<ULevelStreamingPersistenceSettings>()->Properties)
+	{
+		FillClassesProperties(PersistentProperty.Path, PersistentProperty.bIsPublic);
+	}
+	CreateClassDefaultPropertyBag(PropertyType_Private);
+	CreateClassDefaultPropertyBag(PropertyType_Public);
+}
+
+const UPropertyBag* FLevelStreamingPersistentPropertiesInfo::GetPropertyBagFromClass(EPropertyType InAccessSpecifier, const UClass* InClass) const
+{
+	if (ensure(InAccessSpecifier < PropertyType_Count))
+	{
+		const UClass* Class = InClass;
+		while (Class)
+		{
+			if (const FInstancedPropertyBag* InstancedPropertyBag = ObjectClassToPropertyBag[InAccessSpecifier].Find(Class))
+			{
+				return InstancedPropertyBag->GetPropertyBagStruct();
+			}
+			Class = Class->GetSuperClass();
+		}
+	}
+	return nullptr;
+}
+
+const UClass* FLevelStreamingPersistentPropertiesInfo::GetClassFromPropertyBag(EPropertyType InAccessSpecifier, const UPropertyBag* InPropertyBag) const
+{
+	if (ensure(InAccessSpecifier < PropertyType_Count))
+	{
+		if (const TWeakObjectPtr<const UClass>* FoundClass = PropertyBagToObjectClass[InAccessSpecifier].Find(InPropertyBag))
+		{
+			return FoundClass->Get();
+		}
+	}
+	return nullptr;
+}
+
+bool FLevelStreamingPersistentPropertiesInfo::HasProperty(EPropertyType InAccessSpecifier, const FProperty* InProperty) const
+{
+	if (ensure(InAccessSpecifier < PropertyType_Count))
+	{
+		if (const UClass* Class = InProperty->GetOwnerClass())
+		{
+			if (const TSet<FProperty*>* ClassProperties = ClassesProperties[InAccessSpecifier].Find(Class))
+			{
+				return ClassProperties->Contains(InProperty);
+			}
+		}
+	}
+	return false;
+}
+
+void FLevelStreamingPersistentPropertiesInfo::ForEachProperty(EPropertyType InAccessSpecifier, const UClass* InClass, TFunctionRef<void(FProperty*)> Func) const
+{
+	if (ensure(InAccessSpecifier < PropertyType_Count))
+	{
+		const UClass* Class = InClass;
+		while (Class)
+		{
+			if (const TSet<FProperty*>* ClassProperties = ClassesProperties[InAccessSpecifier].Find(Class))
+			{
+				for (FProperty* Property : *ClassProperties)
+				{
+					Func(Property);
+				}
+			}
+			Class = Class->GetSuperClass();
+		}
+	}
+}
+
+bool FLevelStreamingPersistentPropertiesInfo::HasProperties(EPropertyType InAccessSpecifier, const UClass* InClass) const
+{
+	if (ensure(InAccessSpecifier < PropertyType_Count))
+	{
+		const UClass* Class = InClass;
+		while (Class)
+		{
+			if (ClassesProperties[InAccessSpecifier].Contains(Class))
+			{
+				return true;
+			}
+			Class = Class->GetSuperClass();
+		}
+	}
+	return false;
+}
+
+/*
+ * Level Streaming Persistence Manager Console command helper
+ */
+namespace LSPMConsoleCommandHelper
+{
+	template<typename PropertyType>
+	static bool GetValueFromString(const FString& InPropertyValue, PropertyType& OutResult) { return false; }
+
+	template<typename PropertyType>
+	static FString GetValueToString(const PropertyType& InPropertyValue) { return FString(TEXT("<unknown>")); }
+
+	template<typename PropertyType>
+	static bool TrySetPropertyValueFromString(ULevelStreamingPersistenceManager* InLevelStreamingPersistenceManager, const FString& InObjectPath, const FName InPropertyName, const FString& InPropertyValue)
+	{
+		PropertyType ValueFromString;
+		if (GetValueFromString(InPropertyValue, ValueFromString))
+		{
+			if (InLevelStreamingPersistenceManager->TrySetPropertyValue(InObjectPath, InPropertyName, ValueFromString))
+			{
+				PropertyType Result;
+				if (InLevelStreamingPersistenceManager->GetPropertyValue(InObjectPath, InPropertyName, Result))
+				{
+					FString ResultToString = GetValueToString(Result);
+					UE_LOG(LogLevelStreamingPersistenceManager, Log, TEXT("SetPropertyValue succeded : Property[%s] = %s for object %s"), *InPropertyName.ToString(), *ResultToString, *InObjectPath);
+					return true;
+				}
+			}
+		}
+		return false;
+	}
+
+	template<typename PropertyType>
+	static bool GetPropertyValueAsString(ULevelStreamingPersistenceManager* InLevelStreamingPersistenceManager, const FString& InObjectPath, const FName InPropertyName, FString& OutPropertyValue)
+	{
+		PropertyType Result;
+		if (InLevelStreamingPersistenceManager->GetPropertyValue(InObjectPath, InPropertyName, Result))
+		{
+			OutPropertyValue = GetValueToString(Result);
+			UE_LOG(LogLevelStreamingPersistenceManager, Log, TEXT("GetPropertyValue succeded : Property[%s] = %s for object %s"), *InPropertyName.ToString(), *OutPropertyValue, *InObjectPath);
+			return true;
+		}
+		return false;
+	}
+
+	template<> bool GetValueFromString(const FString& InPropertyValue, bool& OutResult) { OutResult = InPropertyValue.ToBool(); return true; }
+	template<> bool GetValueFromString(const FString& InPropertyValue, int32& OutResult) { LexFromString(OutResult, *InPropertyValue); return true; }
+	template<> bool GetValueFromString(const FString& InPropertyValue, int64& OutResult) { LexFromString(OutResult, *InPropertyValue); return true; }
+	template<> bool GetValueFromString(const FString& InPropertyValue, float& OutResult) { LexFromString(OutResult, *InPropertyValue); return true; }
+	template<> bool GetValueFromString(const FString& InPropertyValue, double& OutResult) { LexFromString(OutResult, *InPropertyValue); return true; }
+	template<> bool GetValueFromString(const FString& InPropertyValue, FName& OutResult) { OutResult = FName(InPropertyValue); return true; }
+	template<> bool GetValueFromString(const FString& InPropertyValue, FString& OutResult) { OutResult = InPropertyValue; return true; }
+	template<> bool GetValueFromString(const FString& InPropertyValue, FText& OutResult) { OutResult.FromString(InPropertyValue); return true; }
+	template<> bool GetValueFromString(const FString& InPropertyValue, FVector& OutResult) { return OutResult.InitFromString(InPropertyValue); }
+	template<> bool GetValueFromString(const FString& InPropertyValue, FRotator& OutResult) { return OutResult.InitFromString(InPropertyValue); }
+	template<> bool GetValueFromString(const FString& InPropertyValue, FTransform& OutResult) { return OutResult.InitFromString(InPropertyValue); }
+	template<> bool GetValueFromString(const FString& InPropertyValue, FColor& OutResult) { return OutResult.InitFromString(InPropertyValue); }
+	template<> bool GetValueFromString(const FString& InPropertyValue, FLinearColor& OutResult) { return OutResult.InitFromString(InPropertyValue); }
+
+	template<> FString GetValueToString(const bool& InPropertyValue) { return LexToString(InPropertyValue); }
+	template<> FString GetValueToString(const int32& InPropertyValue) { return LexToString(InPropertyValue); }
+	template<> FString GetValueToString(const int64& InPropertyValue) { return LexToString(InPropertyValue); }
+	template<> FString GetValueToString(const float& InPropertyValue) { return LexToString(InPropertyValue); }
+	template<> FString GetValueToString(const double& InPropertyValue) { return LexToString(InPropertyValue); }
+	template<> FString GetValueToString(const FName& InPropertyValue) { return InPropertyValue.ToString(); }
+	template<> FString GetValueToString(const FString& InPropertyValue) { return InPropertyValue; }
+	template<> FString GetValueToString(const FText& InPropertyValue) { return InPropertyValue.ToString(); }
+	template<> FString GetValueToString(const FVector& InPropertyValue) { return InPropertyValue.ToString(); }
+	template<> FString GetValueToString(const FRotator& InPropertyValue) { return InPropertyValue.ToString(); }
+	template<> FString GetValueToString(const FTransform& InPropertyValue) { return InPropertyValue.ToString(); }
+	template<> FString GetValueToString(const FColor& InPropertyValue) { return InPropertyValue.ToString(); }
+	template<> FString GetValueToString(const FLinearColor& InPropertyValue) { return InPropertyValue.ToString(); }
+
+} // namespace LSPMConsoleCommandHelper
+
+bool ULevelStreamingPersistenceManager::bIsEnabled = true;
+FAutoConsoleVariableRef ULevelStreamingPersistenceManager::EnableCommand(
+	TEXT("s.LevelStreamingPersistence.Enabled"),
+	ULevelStreamingPersistenceManager::bIsEnabled,
+	TEXT("Turn on/off to enable/disable world persistent subsystem."),
+	ECVF_Default);
+
+#if !UE_BUILD_SHIPPING
+
+FAutoConsoleCommand ULevelStreamingPersistenceManager::DumpContentCommand(
+	TEXT("s.LevelStreamingPersistence.DumpContent"),
+	TEXT("Dump the content of WorldPersistentSubsystems"),
+	FConsoleCommandWithArgsDelegate::CreateLambda([](const TArray<FString>& Args)
+	{
+		for (const FWorldContext& Context : GEngine->GetWorldContexts())
+		{
+			UWorld* World = Context.World();
+			if (World && World->IsGameWorld())
+			{
+				if (ULevelStreamingPersistenceManager* LevelStreamingPersistenceManager = World->GetSubsystem<ULevelStreamingPersistenceManager>())
+				{
+					LevelStreamingPersistenceManager->DumpContent();
+				}
+			}
+		}
+	})
+);
+
+FAutoConsoleCommand ULevelStreamingPersistenceManager::SetPropertyValueCommand(
+	TEXT("s.LevelStreamingPersistence.SetPropertyValue"),
+	TEXT("Set the persistent property's value for a given object"),
+	FConsoleCommandWithArgsDelegate::CreateLambda([](const TArray<FString>& InArgs)
+	{
+		if (InArgs.Num() < 3)
+		{
+			return;
+		}
+
+		for (const FWorldContext& Context : GEngine->GetWorldContexts())
+		{
+			UWorld* World = Context.World();
+			if (World && World->IsGameWorld())
+			{
+				if (ULevelStreamingPersistenceManager* Manager = World->GetSubsystem<ULevelStreamingPersistenceManager>())
+				{
+					const FString& ObjectPath = InArgs[0];
+					const FName PropertyName = FName(InArgs[1]);
+					const FString& PropertyValue = InArgs[2];
+
+					if (LSPMConsoleCommandHelper::TrySetPropertyValueFromString<bool>(Manager, ObjectPath, PropertyName, PropertyValue) ||
+						LSPMConsoleCommandHelper::TrySetPropertyValueFromString<int32>(Manager, ObjectPath, PropertyName, PropertyValue) ||
+						LSPMConsoleCommandHelper::TrySetPropertyValueFromString<int64>(Manager, ObjectPath, PropertyName, PropertyValue) ||
+						LSPMConsoleCommandHelper::TrySetPropertyValueFromString<float>(Manager, ObjectPath, PropertyName, PropertyValue) ||
+						LSPMConsoleCommandHelper::TrySetPropertyValueFromString<double>(Manager, ObjectPath, PropertyName, PropertyValue) ||
+						LSPMConsoleCommandHelper::TrySetPropertyValueFromString<FName>(Manager, ObjectPath, PropertyName, PropertyValue) ||
+						LSPMConsoleCommandHelper::TrySetPropertyValueFromString<FString>(Manager, ObjectPath, PropertyName, PropertyValue) ||
+						LSPMConsoleCommandHelper::TrySetPropertyValueFromString<FText>(Manager, ObjectPath, PropertyName, PropertyValue) ||
+						LSPMConsoleCommandHelper::TrySetPropertyValueFromString<FVector>(Manager, ObjectPath, PropertyName, PropertyValue) ||
+						LSPMConsoleCommandHelper::TrySetPropertyValueFromString<FRotator>(Manager, ObjectPath, PropertyName, PropertyValue) ||
+						LSPMConsoleCommandHelper::TrySetPropertyValueFromString<FTransform>(Manager, ObjectPath, PropertyName, PropertyValue) ||
+						LSPMConsoleCommandHelper::TrySetPropertyValueFromString<FColor>(Manager, ObjectPath, PropertyName, PropertyValue) ||
+						LSPMConsoleCommandHelper::TrySetPropertyValueFromString<FLinearColor>(Manager, ObjectPath, PropertyName, PropertyValue))
+					{
+						return;
+					}
+					else if (Manager->TrySetPropertyValueFromString(ObjectPath, PropertyName, PropertyValue))
+					{
+						FString Result;
+						if (Manager->GetPropertyValueAsString(ObjectPath, PropertyName, Result))
+						{
+							UE_LOG(LogLevelStreamingPersistenceManager, Log, TEXT("GetPropertyValueAsString succeded : Property[%s] = %s for object %s"), *PropertyName.ToString(), *Result, *ObjectPath);
+							return;
+						}
+					}
+				}
+			}
+		}
+	})
+);
+
+FAutoConsoleCommand ULevelStreamingPersistenceManager::GetPropertyValueCommand(
+	TEXT("s.LevelStreamingPersistence.GetPropertyValue"),
+	TEXT("Get the persistent property's value for a given object"),
+	FConsoleCommandWithArgsDelegate::CreateLambda([](const TArray<FString>& InArgs)
+	{
+		if (InArgs.Num() < 2)
+		{
+			return;
+		}
+
+		for (const FWorldContext& Context : GEngine->GetWorldContexts())
+		{
+			UWorld* World = Context.World();
+			if (World && World->IsGameWorld())
+			{
+				if (ULevelStreamingPersistenceManager* Manager = World->GetSubsystem<ULevelStreamingPersistenceManager>())
+				{
+					const FString& ObjectPath = InArgs[0];
+					const FName PropertyName = FName(InArgs[1]);
+
+					FString PropertyValue;
+					if (LSPMConsoleCommandHelper::GetPropertyValueAsString<bool>(Manager, ObjectPath, PropertyName, PropertyValue) ||
+						LSPMConsoleCommandHelper::GetPropertyValueAsString<int32>(Manager, ObjectPath, PropertyName, PropertyValue) ||
+						LSPMConsoleCommandHelper::GetPropertyValueAsString<int64>(Manager, ObjectPath, PropertyName, PropertyValue) ||
+						LSPMConsoleCommandHelper::GetPropertyValueAsString<float>(Manager, ObjectPath, PropertyName, PropertyValue) ||
+						LSPMConsoleCommandHelper::GetPropertyValueAsString<double>(Manager, ObjectPath, PropertyName, PropertyValue) ||
+						LSPMConsoleCommandHelper::GetPropertyValueAsString<FName>(Manager, ObjectPath, PropertyName, PropertyValue) ||
+						LSPMConsoleCommandHelper::GetPropertyValueAsString<FString>(Manager, ObjectPath, PropertyName, PropertyValue) ||
+						LSPMConsoleCommandHelper::GetPropertyValueAsString<FText>(Manager, ObjectPath, PropertyName, PropertyValue) ||
+						LSPMConsoleCommandHelper::GetPropertyValueAsString<FVector>(Manager, ObjectPath, PropertyName, PropertyValue) ||
+						LSPMConsoleCommandHelper::GetPropertyValueAsString<FRotator>(Manager, ObjectPath, PropertyName, PropertyValue) ||
+						LSPMConsoleCommandHelper::GetPropertyValueAsString<FTransform>(Manager, ObjectPath, PropertyName, PropertyValue) ||
+						LSPMConsoleCommandHelper::GetPropertyValueAsString<FColor>(Manager, ObjectPath, PropertyName, PropertyValue) ||
+						LSPMConsoleCommandHelper::GetPropertyValueAsString<FLinearColor>(Manager, ObjectPath, PropertyName, PropertyValue))
+					{
+						return;
+					}
+					else
+					{
+						FString Result;
+						if (Manager->GetPropertyValueAsString(ObjectPath, PropertyName, Result))
+						{
+							UE_LOG(LogLevelStreamingPersistenceManager, Log, TEXT("GetPropertyValueAsString succeded : Property[%s] = %s for object %s"), *PropertyName.ToString(), *Result, *ObjectPath);
+							return;
+						}
+					}
+				}
+			}
+		}
+	})
+);
+
+#endif // !UE_BUILD_SHIPPING
\ No newline at end of file
diff --git a/Engine/Plugins/Runtime/LevelStreamingPersistence/Source/LevelStreamingPersistence/Private/LevelStreamingPersistenceModule.cpp b/Engine/Plugins/Runtime/LevelStreamingPersistence/Source/LevelStreamingPersistence/Private/LevelStreamingPersistenceModule.cpp
new file mode 100644
--- /dev/null
+++ b/Engine/Plugins/Runtime/LevelStreamingPersistence/Source/LevelStreamingPersistence/Private/LevelStreamingPersistenceModule.cpp
@@ -0,0 +1,12 @@
+// Copyright Epic Games, Inc. All Rights Reserved.
+
+#include "LevelStreamingPersistenceModule.h"
+
+class FLevelStreamingPersistenceModule : public ILevelStreamingPersistenceModule
+{
+public:
+	virtual void StartupModule() override {};
+	virtual void ShutdownModule() override {};
+};
+
+IMPLEMENT_MODULE(FLevelStreamingPersistenceModule, LevelStreamingPersistence);
\ No newline at end of file
diff --git a/Engine/Plugins/Runtime/LevelStreamingPersistence/Source/LevelStreamingPersistence/Private/LevelStreamingPersistentObjectPropertyBag.cpp b/Engine/Plugins/Runtime/LevelStreamingPersistence/Source/LevelStreamingPersistence/Private/LevelStreamingPersistentObjectPropertyBag.cpp
new file mode 100644
--- /dev/null
+++ b/Engine/Plugins/Runtime/LevelStreamingPersistence/Source/LevelStreamingPersistence/Private/LevelStreamingPersistentObjectPropertyBag.cpp
@@ -0,0 +1,179 @@
+// Copyright Epic Games, Inc. All Rights Reserved.
+
+#include "LevelStreamingPersistentObjectPropertyBag.h"
+#include "UObject/PropertyAccessUtil.h"
+#include "PropertyPathHelpers.h"
+
+bool FLevelStreamingPersistentObjectPropertyBag::Initialize(TFunctionRef<const UPropertyBag* ()> Func)
+{
+	if (IsValid())
+	{
+		return true;
+	}
+	else
+	{
+		const UPropertyBag* DefaultClass = Func();
+		if (ensure(DefaultClass))
+		{
+			PropertyBag.InitializeFromBagStruct(DefaultClass);
+			return true;
+		}
+	}
+	return false;
+}
+
+const FProperty* FLevelStreamingPersistentObjectPropertyBag::FindPropertyByName(const FName InPropertyName) const
+{
+	const FPropertyBagPropertyDesc* Desc = PropertyBag.FindPropertyDescByName(InPropertyName);
+	return Desc ? Desc->CachedProperty : nullptr;
+}
+
+bool FLevelStreamingPersistentObjectPropertyBag::ComparePropertyValueWithObject(const UObject* InObject, const FProperty* InObjectProperty, bool& bOutIsIdentical) const
+{
+	check(IsValid());
+	bOutIsIdentical = false;
+	const FProperty* BagProperty = InObjectProperty ? FindPropertyByName(InObjectProperty->GetFName()) : nullptr;
+	if (!BagProperty || !PropertyAccessUtil::ArePropertiesCompatible(BagProperty, InObjectProperty))
+	{
+		return false;
+	}
+	const void* ObjectValue = InObjectProperty->ContainerPtrToValuePtr<void>(InObject);
+	const void* BagValue = BagProperty->ContainerPtrToValuePtr<void>(PropertyBag.GetValue().GetMemory());
+	bOutIsIdentical = PropertyAccessUtil::IsSinglePropertyIdentical(InObjectProperty, ObjectValue, BagProperty, BagValue);
+	return true;
+}
+
+const FProperty* FLevelStreamingPersistentObjectPropertyBag::CopyPropertyValueFromObject(const UObject* InObject, const FProperty* InObjectProperty)
+{
+	check(IsValid());
+	const FProperty* BagProperty = InObjectProperty ? FindPropertyByName(InObjectProperty->GetFName()) : nullptr;
+	if (!BagProperty)
+	{
+		return nullptr;
+	}
+
+	const void* BagValue = BagProperty->ContainerPtrToValuePtr<void>(PropertyBag.GetValue().GetMemory());
+	EPropertyAccessResultFlags Result = PropertyAccessUtil::GetPropertyValue_Object(InObjectProperty, InObject, BagProperty, const_cast<void*>(BagValue), INDEX_NONE);
+	return (Result == EPropertyAccessResultFlags::Success) ? BagProperty : nullptr;
+}
+
+const FProperty* FLevelStreamingPersistentObjectPropertyBag::CopyPropertyValueToObject(UObject* InObject, const FProperty* InObjectProperty) const
+{
+	check(IsValid());
+	const FProperty* BagProperty = ::IsValid(InObject) && InObjectProperty ? FindPropertyByName(InObjectProperty->GetFName()) : nullptr;
+	if (!BagProperty)
+	{
+		return nullptr;
+	}
+	const void* ObjectValue = InObjectProperty->ContainerPtrToValuePtr<void>(InObject);
+	EPropertyAccessResultFlags Result = PropertyAccessUtil::GetPropertyValue_InContainer(BagProperty, PropertyBag.GetValue().GetMemory(), InObjectProperty, const_cast<void*>(ObjectValue), INDEX_NONE);
+	return (Result == EPropertyAccessResultFlags::Success) ? BagProperty : nullptr;
+}
+
+void FLevelStreamingPersistentObjectPropertyBag::ForEachProperty(TFunctionRef<void(const FProperty*)> Func) const
+{
+	check(IsValid());
+	for (const FPropertyBagPropertyDesc& SourceDesc : PropertyBag.GetPropertyBagStruct()->GetPropertyDescs())
+	{
+		if (const FProperty* Property = SourceDesc.CachedProperty)
+		{
+			Func(Property);
+		}
+	}
+}
+
+bool FLevelStreamingPersistentObjectPropertyBag::GetPropertyValueAsString(const FName InPropertyName, FString& OutPropertyValueAsString) const
+{
+	if (const FProperty* Property = IsValid() ? FindPropertyByName(InPropertyName) : nullptr)
+	{
+		return PropertyPathHelpers::GetPropertyValueAsString(const_cast<uint8*>(PropertyBag.GetValue().GetMemory()), const_cast<UPropertyBag*>(GetPropertyBagStruct()), InPropertyName.ToString(), OutPropertyValueAsString);
+	}
+	return false;
+}
+
+bool FLevelStreamingPersistentObjectPropertyBag::SetPropertyValueFromString(const FName InPropertyName, const FString& InPropertyValue)
+{
+	if (const FProperty* Property = IsValid() ? FindPropertyByName(InPropertyName) : nullptr)
+	{
+		return PropertyPathHelpers::SetPropertyValueFromString(const_cast<uint8*>(PropertyBag.GetValue().GetMemory()), const_cast<UPropertyBag*>(GetPropertyBagStruct()), InPropertyName.ToString(), InPropertyValue);
+	}
+	return false;
+}
+
+void FLevelStreamingPersistentObjectPropertyBag::DumpContent(TFunctionRef<void(const FProperty*, const FString&)> Func, const TArray<const FProperty*>* InDumpProperties)
+{
+	auto DumpProperty = [this, &Func](const FProperty* Property)
+	{
+		FString PropertyValue;
+		if (GetPropertyValueAsString(Property->GetFName(), PropertyValue))
+		{
+			Func(Property, *PropertyValue);
+		}
+	};
+
+	if (InDumpProperties)
+	{
+		for (const FProperty* Property : *InDumpProperties)
+		{
+			DumpProperty(Property);
+		}
+	}
+	else
+	{
+		ForEachProperty([this, &DumpProperty](const FProperty* Property)
+		{
+			DumpProperty(Property);
+		});
+	}
+}
+
+template<> struct FLevelStreamingPersistentObjectPropertyBag::TPropertyBagPropertyType<bool> { static const EPropertyBagPropertyType Value = EPropertyBagPropertyType::Bool; };
+template<> struct FLevelStreamingPersistentObjectPropertyBag::TPropertyBagPropertyType<int32> { static const EPropertyBagPropertyType Value = EPropertyBagPropertyType::Int32; };
+template<> struct FLevelStreamingPersistentObjectPropertyBag::TPropertyBagPropertyType<int64> { static const EPropertyBagPropertyType Value = EPropertyBagPropertyType::Int64; };
+template<> struct FLevelStreamingPersistentObjectPropertyBag::TPropertyBagPropertyType<float> { static const EPropertyBagPropertyType Value = EPropertyBagPropertyType::Float; };
+template<> struct FLevelStreamingPersistentObjectPropertyBag::TPropertyBagPropertyType<double> { static const EPropertyBagPropertyType Value = EPropertyBagPropertyType::Double; };
+template<> struct FLevelStreamingPersistentObjectPropertyBag::TPropertyBagPropertyType<FName> { static const EPropertyBagPropertyType Value = EPropertyBagPropertyType::Name; };
+template<> struct FLevelStreamingPersistentObjectPropertyBag::TPropertyBagPropertyType<FString> { static const EPropertyBagPropertyType Value = EPropertyBagPropertyType::String; };
+template<> struct FLevelStreamingPersistentObjectPropertyBag::TPropertyBagPropertyType<FText> { static const EPropertyBagPropertyType Value = EPropertyBagPropertyType::Text; };
+template<> struct FLevelStreamingPersistentObjectPropertyBag::TPropertyBagPropertyType<FSoftObjectPtr> { static const EPropertyBagPropertyType Value = EPropertyBagPropertyType::SoftObject; };
+template<> struct FLevelStreamingPersistentObjectPropertyBag::TPropertyBagPropertyType<UObject*> { static const EPropertyBagPropertyType Value = EPropertyBagPropertyType::Object; };
+template<> struct FLevelStreamingPersistentObjectPropertyBag::TPropertyBagPropertyType<FVector> { static const EPropertyBagPropertyType Value = EPropertyBagPropertyType::Struct; };
+template<> struct FLevelStreamingPersistentObjectPropertyBag::TPropertyBagPropertyType<FRotator> { static const EPropertyBagPropertyType Value = EPropertyBagPropertyType::Struct; };
+template<> struct FLevelStreamingPersistentObjectPropertyBag::TPropertyBagPropertyType<FTransform> { static const EPropertyBagPropertyType Value = EPropertyBagPropertyType::Struct; };
+template<> struct FLevelStreamingPersistentObjectPropertyBag::TPropertyBagPropertyType<FColor> { static const EPropertyBagPropertyType Value = EPropertyBagPropertyType::Struct; };
+template<> struct FLevelStreamingPersistentObjectPropertyBag::TPropertyBagPropertyType<FLinearColor> { static const EPropertyBagPropertyType Value = EPropertyBagPropertyType::Struct; };
+template<> struct FLevelStreamingPersistentObjectPropertyBag::TPropertyBagPropertyType<FStructView> { static const EPropertyBagPropertyType Value = EPropertyBagPropertyType::Struct; };
+template<> struct FLevelStreamingPersistentObjectPropertyBag::TPropertyBagPropertyType<FConstStructView> { static const EPropertyBagPropertyType Value = EPropertyBagPropertyType::Struct; };
+
+template<> EPropertyBagResult FLevelStreamingPersistentObjectPropertyBag::SetValue(const FName InPropertyName, const bool& InPropertyValue) { return PropertyBag.SetValueBool(InPropertyName, InPropertyValue); }
+template<> EPropertyBagResult FLevelStreamingPersistentObjectPropertyBag::SetValue(const FName InPropertyName, const int32& InPropertyValue) { return PropertyBag.SetValueInt32(InPropertyName, InPropertyValue); }
+template<> EPropertyBagResult FLevelStreamingPersistentObjectPropertyBag::SetValue(const FName InPropertyName, const int64& InPropertyValue) { return PropertyBag.SetValueInt64(InPropertyName, InPropertyValue); }
+template<> EPropertyBagResult FLevelStreamingPersistentObjectPropertyBag::SetValue(const FName InPropertyName, const float& InPropertyValue) { return PropertyBag.SetValueFloat(InPropertyName, InPropertyValue); }
+template<> EPropertyBagResult FLevelStreamingPersistentObjectPropertyBag::SetValue(const FName InPropertyName, const double& InPropertyValue) { return PropertyBag.SetValueDouble(InPropertyName, InPropertyValue); }
+template<> EPropertyBagResult FLevelStreamingPersistentObjectPropertyBag::SetValue(const FName InPropertyName, const FName& InPropertyValue) { return PropertyBag.SetValueName(InPropertyName, InPropertyValue); }
+template<> EPropertyBagResult FLevelStreamingPersistentObjectPropertyBag::SetValue(const FName InPropertyName, const FString& InPropertyValue) { return PropertyBag.SetValueString(InPropertyName, InPropertyValue); }
+template<> EPropertyBagResult FLevelStreamingPersistentObjectPropertyBag::SetValue(const FName InPropertyName, const FText& InPropertyValue) { return PropertyBag.SetValueText(InPropertyName, InPropertyValue); }
+template<> EPropertyBagResult FLevelStreamingPersistentObjectPropertyBag::SetValue(const FName InPropertyName, const FSoftObjectPtr& InPropertyValue) { return PropertyBag.SetValueObject(InPropertyName, InPropertyValue.Get()); }
+template<> EPropertyBagResult FLevelStreamingPersistentObjectPropertyBag::SetValue(const FName InPropertyName, const UObject& InPropertyValue) { return PropertyBag.SetValueObject(InPropertyName, &InPropertyValue); }
+template<> EPropertyBagResult FLevelStreamingPersistentObjectPropertyBag::SetValue(const FName InPropertyName, const FConstStructView& InPropertyValue) { return PropertyBag.SetValueStruct(InPropertyName, InPropertyValue); }
+template<> EPropertyBagResult FLevelStreamingPersistentObjectPropertyBag::SetValue(const FName InPropertyName, const FVector& InPropertyValue) { return SetValue(InPropertyName, FConstStructView::Make(InPropertyValue)); }
+template<> EPropertyBagResult FLevelStreamingPersistentObjectPropertyBag::SetValue(const FName InPropertyName, const FRotator& InPropertyValue) { return SetValue(InPropertyName, FConstStructView::Make(InPropertyValue)); }
+template<> EPropertyBagResult FLevelStreamingPersistentObjectPropertyBag::SetValue(const FName InPropertyName, const FTransform& InPropertyValue) { return SetValue(InPropertyName, FConstStructView::Make(InPropertyValue)); }
+template<> EPropertyBagResult FLevelStreamingPersistentObjectPropertyBag::SetValue(const FName InPropertyName, const FColor& InPropertyValue) { return SetValue(InPropertyName, FConstStructView::Make(InPropertyValue)); }
+template<> EPropertyBagResult FLevelStreamingPersistentObjectPropertyBag::SetValue(const FName InPropertyName, const FLinearColor& InPropertyValue) { return SetValue(InPropertyName, FConstStructView::Make(InPropertyValue)); }
+
+template<> TValueOrError<bool, EPropertyBagResult> FLevelStreamingPersistentObjectPropertyBag::GetValue(const FName InPropertyName) const { return PropertyBag.GetValueBool(InPropertyName); }
+template<> TValueOrError<int32, EPropertyBagResult> FLevelStreamingPersistentObjectPropertyBag::GetValue(const FName InPropertyName) const { return PropertyBag.GetValueInt32(InPropertyName); }
+template<> TValueOrError<int64, EPropertyBagResult> FLevelStreamingPersistentObjectPropertyBag::GetValue(const FName InPropertyName) const { return PropertyBag.GetValueInt64(InPropertyName); }
+template<> TValueOrError<float, EPropertyBagResult> FLevelStreamingPersistentObjectPropertyBag::GetValue(const FName InPropertyName) const { return PropertyBag.GetValueFloat(InPropertyName); }
+template<> TValueOrError<double, EPropertyBagResult> FLevelStreamingPersistentObjectPropertyBag::GetValue(const FName InPropertyName) const { return PropertyBag.GetValueDouble(InPropertyName); }
+template<> TValueOrError<FName, EPropertyBagResult> FLevelStreamingPersistentObjectPropertyBag::GetValue(const FName InPropertyName) const { return PropertyBag.GetValueName(InPropertyName); }
+template<> TValueOrError<FString, EPropertyBagResult> FLevelStreamingPersistentObjectPropertyBag::GetValue(const FName InPropertyName) const { return PropertyBag.GetValueString(InPropertyName); }
+template<> TValueOrError<FText, EPropertyBagResult> FLevelStreamingPersistentObjectPropertyBag::GetValue(const FName InPropertyName) const { return PropertyBag.GetValueText(InPropertyName); }
+template<> TValueOrError<UObject*, EPropertyBagResult> FLevelStreamingPersistentObjectPropertyBag::GetValue(const FName InPropertyName) const { return PropertyBag.GetValueObject(InPropertyName); }
+template<> TValueOrError<FStructView, EPropertyBagResult> FLevelStreamingPersistentObjectPropertyBag::GetValue(const FName InPropertyName) const { return PropertyBag.GetValueStruct(InPropertyName); }
+template<> TValueOrError<FVector, EPropertyBagResult> FLevelStreamingPersistentObjectPropertyBag::GetValue(const FName InPropertyName) const { return GetValueStruct<FVector>(InPropertyName); }
+template<> TValueOrError<FRotator, EPropertyBagResult> FLevelStreamingPersistentObjectPropertyBag::GetValue(const FName InPropertyName) const { return GetValueStruct<FRotator>(InPropertyName); }
+template<> TValueOrError<FTransform, EPropertyBagResult> FLevelStreamingPersistentObjectPropertyBag::GetValue(const FName InPropertyName) const { return GetValueStruct<FTransform>(InPropertyName); }
+template<> TValueOrError<FColor, EPropertyBagResult> FLevelStreamingPersistentObjectPropertyBag::GetValue(const FName InPropertyName) const { return GetValueStruct<FColor>(InPropertyName); }
+template<> TValueOrError<FLinearColor, EPropertyBagResult> FLevelStreamingPersistentObjectPropertyBag::GetValue(const FName InPropertyName) const { return GetValueStruct<FLinearColor>(InPropertyName); }
\ No newline at end of file
diff --git a/Engine/Plugins/Runtime/LevelStreamingPersistence/Source/LevelStreamingPersistence/Public/LevelStreamingPersistenceManager.h b/Engine/Plugins/Runtime/LevelStreamingPersistence/Source/LevelStreamingPersistence/Public/LevelStreamingPersistenceManager.h
new file mode 100644
--- /dev/null
+++ b/Engine/Plugins/Runtime/LevelStreamingPersistence/Source/LevelStreamingPersistence/Public/LevelStreamingPersistenceManager.h
@@ -0,0 +1,227 @@
+// Copyright Epic Games, Inc. All Rights Reserved.
+
+#pragma once
+
+#include "CoreTypes.h"
+#include "PropertyBag.h"
+#include "Containers/StaticArray.h"
+#include "UObject/WeakObjectPtrTemplates.h"
+#include "Subsystems/WorldSubsystem.h"
+#include "LevelStreamingPersistentObjectPropertyBag.h"
+#include "LevelStreamingPersistenceModule.h"
+#include "LevelStreamingPersistenceSettings.h"
+#include "LevelStreamingPersistenceManager.generated.h"
+
+class ULevel;
+class UWorld;
+class ULevelStreaming;
+using FLevelStreamingPersistentPropertyArray = TArray<FCustomPropertyListNode, TInlineAllocator<32>>;
+
+struct FLevelStreamingPersistentObjectPrivateProperties
+{
+	// Payload data containing serialized changed properties
+	TArray<uint8> PayloadData;
+	// Initial snapshot used detect changed properties
+	FLevelStreamingPersistentObjectPropertyBag Snapshot;
+	// Persistent Properties
+	TArray<const FProperty*> PersistentProperties;
+	
+	void ForEachPersistentProperty(TFunctionRef<void(const FProperty*)> Func) const;
+	void BuildSerializedPropertyList(FLevelStreamingPersistentPropertyArray& OutCustomPropertyList) const;
+};
+
+struct FLevelStreamingPersistentObjectPublicProperties
+{
+	// Public properties
+	FLevelStreamingPersistentObjectPropertyBag PropertyBag;
+	// Persistent Properties
+	TArray<const FProperty*> PersistentProperties;
+
+	void ForEachPersistentProperty(TFunctionRef<void(const FProperty*)> Func) const;
+};
+
+struct FLevelStreamingPersistentPropertyValues
+{
+	TMap<FString, FLevelStreamingPersistentObjectPrivateProperties> ObjectsPrivatePropertyValues;
+	TMap<FString, FLevelStreamingPersistentObjectPublicProperties> ObjectsPublicPropertyValues;
+};
+
+// Helper class to access FLevelStreamingPersistentProperty's Properties
+class FLevelStreamingPersistentPropertiesInfo
+{
+public:
+	void Initialize();
+
+	enum EPropertyType
+	{
+		PropertyType_Public,
+		PropertyType_Private,
+		PropertyType_Count
+	};
+
+	// Configurable persistent properties helper methods
+	const UPropertyBag* GetPropertyBagFromClass(EPropertyType InAccessSpecifier, const UClass* InClass) const;
+	const UClass* GetClassFromPropertyBag(EPropertyType InAccessSpecifier, const UPropertyBag* InPropertyBag) const;
+	void ForEachProperty(EPropertyType InAccessSpecifier, const UClass* InClass, TFunctionRef<void(FProperty*)> Func) const;
+	bool HasProperties(EPropertyType InAccessSpecifier, const UClass* InClass) const;
+	bool HasProperty(EPropertyType InAccessSpecifier, const FProperty* InProperty) const;
+
+private:
+	/* Acceleration maps to find properties for a given class */
+	TStaticArray<TMap<TWeakObjectPtr<const UClass>, TSet<FProperty*>>, PropertyType_Count> ClassesProperties;
+	TStaticArray<TMap<TWeakObjectPtr<const UClass>, FInstancedPropertyBag>, PropertyType_Count> ObjectClassToPropertyBag;
+	TStaticArray<TMap<const UPropertyBag*, TWeakObjectPtr<const UClass>>, PropertyType_Count> PropertyBagToObjectClass;
+};
+
+UCLASS()
+class LEVELSTREAMINGPERSISTENCE_API ULevelStreamingPersistenceManager : public UWorldSubsystem
+{
+	GENERATED_BODY()
+
+public:
+	//~ Begin WorldSubsystem
+	ULevelStreamingPersistenceManager(const FObjectInitializer&);
+	~ULevelStreamingPersistenceManager() {}
+	virtual bool ShouldCreateSubsystem(UObject* Outer) const override;
+	virtual void Initialize(FSubsystemCollectionBase& Collection) override;
+	virtual void PostInitialize() override;
+	virtual void Deinitialize() override;
+	//~ End WorldSubsystem
+
+	// Sets property value and creates the entry if necessary, returns true on success.
+	template <typename ClassType, typename PropertyType>
+	bool SetPropertyValue(const FString& InObjectPathName, const FName InPropertyName, const PropertyType& InPropertyValue);
+
+	// Sets the property value on existing entries, returns true on success.
+	template <typename PropertyType>
+	bool TrySetPropertyValue(const FString& InObjectPathName, const FName InPropertyName, const PropertyType& InPropertyValue);
+
+	// Gets the property value if found, returns true on success.
+	template<typename PropertyType>
+	bool GetPropertyValue(const FString& InObjectPathName, const FName InPropertyName, PropertyType& OutPropertyValue) const;
+
+	// Sets the property value converted from the provided string value on existing entries, returns true on success.
+	bool TrySetPropertyValueFromString(const FString& InObjectPathName, const FName InPropertyName, const FString& InPropertyValue);
+
+	// Gets the property value and converts it to a string if found, returns true on success.
+	bool GetPropertyValueAsString(const FString& InObjectPathName, const FName InPropertyName, FString& OutPropertyValue);
+
+protected:
+	virtual bool DoesSupportWorldType(const EWorldType::Type WorldType) const;
+
+private:
+	// Level streaming visibility callbacks
+	void OnLevelBeginMakingInvisible(UWorld* World, const ULevelStreaming* StreamingLevel, ULevel* LoadedLevel);
+	void OnLevelBeginMakingVisible(UWorld* World, const ULevelStreaming* StreamingLevel, ULevel* LoadedLevel);
+
+	// Save persistent properties
+	bool SaveLevelPersistentPropertyValues(const ULevel* InLevel);
+
+	// Restore persistent properties
+	bool RestoreLevelPersistentPropertyValues(ULevel* InLevel) const;
+	int32 RestorePrivateProperties(UObject* InObject, const FLevelStreamingPersistentObjectPrivateProperties& InPersistentProperties) const;
+	int32 RestorePublicProperties(UObject* InObject, const FLevelStreamingPersistentObjectPublicProperties& InPersistentProperties) const;
+
+	// Snapshot of persistent properties
+	bool BuildSnapshot(const UObject* InObject, FLevelStreamingPersistentObjectPropertyBag& OutSnapshot) const;
+	bool DiffWithSnapshot(const UObject* InObject, const FLevelStreamingPersistentObjectPropertyBag& InSnapshot, TArray<const FProperty*>& OutChangedProperties) const;
+
+	bool IsEnabled() const { return bIsEnabled; }
+	bool SplitObjectPath(const FString& InObjectPathName, FString& OutLevelPathName, FString& OutShortObjectPathName) const;
+	const FString GetResolvedObjectPathName(const FString& InObjectPathName) const;
+	bool CopyPropertyBagValueToObject(const FLevelStreamingPersistentObjectPropertyBag* InPropertyBag, UObject* InObject, FProperty* InObjectProperty) const;
+	TPair<UObject*, FProperty*> GetObjectPropertyPair(const FString& InObjectPathName, const FName InPropertyName) const;
+	void DumpContent() const;
+	const FLevelStreamingPersistentObjectPropertyBag* GetPropertyBag(const FString& InObjectPathName) const;
+	FLevelStreamingPersistentObjectPropertyBag* GetPropertyBag(const FString& InObjectPathName);
+
+	template <typename ClassType>
+	FLevelStreamingPersistentObjectPropertyBag* GetOrCreatePropertyBag(const FString& InObjectPathName, const FName InPropertyName);
+
+	// Console commands
+	static class FAutoConsoleVariableRef EnableCommand;
+	static bool bIsEnabled;
+#if !UE_BUILD_SHIPPING
+	static class FAutoConsoleCommand DumpContentCommand;
+	static class FAutoConsoleCommand SetPropertyValueCommand;
+	static class FAutoConsoleCommand GetPropertyValueCommand;
+#endif
+
+	// Per-level Persistent property values
+	mutable TMap<FString, FLevelStreamingPersistentPropertyValues> LevelsPropertyValues;
+
+	// Persistence Module
+	ILevelStreamingPersistenceModule* PersistenceModule = nullptr;
+
+	// Persistent Properties Info
+	FLevelStreamingPersistentPropertiesInfo PersistentPropertiesInfo;
+
+	friend class FPersistenceModule;
+};
+
+template <typename PropertyType>
+bool ULevelStreamingPersistenceManager::TrySetPropertyValue(const FString& InObjectPathName, const FName InPropertyName, const PropertyType& InPropertyValue)
+{
+	const FString ObjectPathName = GetResolvedObjectPathName(InObjectPathName);
+	FLevelStreamingPersistentObjectPropertyBag* PropertyBag = GetPropertyBag(ObjectPathName);
+	if (PropertyBag && PropertyBag->SetPropertyValue(InPropertyName, InPropertyValue))
+	{
+		// If object is loaded, copy property value to object
+		auto ObjectPropertyPair = GetObjectPropertyPair(ObjectPathName, InPropertyName);
+		CopyPropertyBagValueToObject(PropertyBag, ObjectPropertyPair.Key, ObjectPropertyPair.Value);
+		return true;
+	}
+	return false;
+}
+
+template <typename ClassType, typename PropertyType>
+bool ULevelStreamingPersistenceManager::SetPropertyValue(const FString& InObjectPathName, const FName InPropertyName, const PropertyType& InPropertyValue)
+{
+	const FString ObjectPathName = GetResolvedObjectPathName(InObjectPathName);
+	FLevelStreamingPersistentObjectPropertyBag* PropertyBag = GetOrCreatePropertyBag<ClassType>(ObjectPathName, InPropertyName);
+	if (PropertyBag && PropertyBag->SetPropertyValue(InPropertyName, InPropertyValue))
+	{
+		// If object is loaded, copy property value to object
+		auto ObjectPropertyPair = GetObjectPropertyPair(ObjectPathName, InPropertyName);
+		CopyPropertyBagValueToObject(PropertyBag, ObjectPropertyPair.Key, ObjectPropertyPair.Value);
+		return true;
+	}
+	return false;
+}
+
+template<typename PropertyType>
+bool ULevelStreamingPersistenceManager::GetPropertyValue(const FString& InObjectPathName, const FName InPropertyName, PropertyType& OutPropertyValue) const
+{
+	const FString ObjectPathName = GetResolvedObjectPathName(InObjectPathName);
+	const FLevelStreamingPersistentObjectPropertyBag* PropertyBag = GetPropertyBag(ObjectPathName);
+	return PropertyBag ? PropertyBag->GetPropertyValue(InPropertyName, OutPropertyValue) : false;
+}
+
+template <typename ClassType>
+FLevelStreamingPersistentObjectPropertyBag* ULevelStreamingPersistenceManager::GetOrCreatePropertyBag(const FString& InObjectPathName, const FName InPropertyName)
+{
+	const FString ObjectPathName = GetResolvedObjectPathName(InObjectPathName);
+	const UClass* ObjectClass = ClassType::StaticClass();
+
+	FString LevelPathName;
+	FString ObjectShortPathName;
+	if (!SplitObjectPath(ObjectPathName, LevelPathName, ObjectShortPathName))
+	{
+		return nullptr;
+	}
+
+	if (FProperty* ObjectProperty = ObjectClass->FindPropertyByName(InPropertyName))
+	{
+		if (PersistentPropertiesInfo.HasProperty(FLevelStreamingPersistentPropertiesInfo::PropertyType_Public, ObjectProperty))
+		{
+			FLevelStreamingPersistentPropertyValues& LevelProperties = LevelsPropertyValues.FindOrAdd(LevelPathName);
+
+			FLevelStreamingPersistentObjectPublicProperties& ObjectPublicPropertyValues = LevelProperties.ObjectsPublicPropertyValues.FindOrAdd(ObjectPathName);
+			if (ObjectPublicPropertyValues.PropertyBag.Initialize([this, ObjectClass]() { return PersistentPropertiesInfo.GetPropertyBagFromClass(FLevelStreamingPersistentPropertiesInfo::PropertyType_Public, ObjectClass); }))
+			{
+				return &ObjectPublicPropertyValues.PropertyBag;
+			}
+		}
+	}
+	return nullptr;
+}
\ No newline at end of file
diff --git a/Engine/Plugins/Runtime/LevelStreamingPersistence/Source/LevelStreamingPersistence/Public/LevelStreamingPersistenceModule.h b/Engine/Plugins/Runtime/LevelStreamingPersistence/Source/LevelStreamingPersistence/Public/LevelStreamingPersistenceModule.h
new file mode 100644
--- /dev/null
+++ b/Engine/Plugins/Runtime/LevelStreamingPersistence/Source/LevelStreamingPersistence/Public/LevelStreamingPersistenceModule.h
@@ -0,0 +1,97 @@
+// Copyright Epic Games, Inc. All Rights Reserved.
+
+#pragma once
+
+#include "CoreMinimal.h"
+#include "Modules/ModuleInterface.h"
+#include "Modules/ModuleManager.h"
+#include "UObject/Class.h"
+
+/**
+ * The public interface to this module.  In most cases, this interface is only public to sibling modules 
+ * within this plugin.
+ */
+class LEVELSTREAMINGPERSISTENCE_API ILevelStreamingPersistenceModule : public IModuleInterface
+{
+public:
+	DECLARE_DELEGATE_RetVal_TwoParams(bool, FShouldPersistProperty, const UObject*, const FProperty*);
+	DECLARE_DELEGATE_RetVal_TwoParams(void, FPostRestorePersistedProperty, const UObject*, const FProperty*);
+
+	template<typename ClassType>
+	FShouldPersistProperty& OnShouldPersistProperty()
+	{
+		return OnShouldPersistProperty(ClassType::StaticClass());
+	}
+
+	FShouldPersistProperty& OnShouldPersistProperty(const UClass* InClass)
+	{
+		return ClassShouldPersistProperty.FindOrAdd(InClass);
+	}
+
+	template<typename ClassType>
+	FPostRestorePersistedProperty& OnPostRestorePersistedProperty()
+	{
+		return OnPostRestorePersistedProperty(ClassType::StaticClass());
+	}
+
+	FPostRestorePersistedProperty& OnPostRestorePersistedProperty(const UClass* InClass)
+	{
+		return ClassPostRestorePersistedProperty.FindOrAdd(InClass);
+	}
+
+	/**
+	 * Singleton-like access to this module's interface.  This is just for convenience!
+	 * Beware of calling this during the shutdown phase, though.  Your module might have been unloaded already.
+	 *
+	 * @return Returns singleton instance, loading the module on demand if needed
+	 */
+	static inline ILevelStreamingPersistenceModule& Get()
+	{
+		return FModuleManager::LoadModuleChecked<ILevelStreamingPersistenceModule>("LevelStreamingPersistence");
+	}
+
+	/**
+	 * Checks to see if this module is loaded and ready.  It is only valid to call Get() if IsAvailable() returns true.
+	 *
+	 * @return True if the module is loaded and ready to use
+	 */
+	static inline bool IsAvailable()
+	{
+		return FModuleManager::Get().IsModuleLoaded("LevelStreamingPersistence");
+	}
+
+private:
+
+	bool ShouldPersistProperty(const UObject* InObject, const FProperty* InProperty) const
+	{
+		const UClass* Class = InObject->GetClass();
+		while (Class)
+		{
+			const FShouldPersistProperty* Found = ClassShouldPersistProperty.Find(Class);
+			if (Found && Found->IsBound() && !Found->Execute(InObject, InProperty))
+			{
+				return false;
+			}
+			Class = Class->GetSuperClass();
+		};
+		return true;
+	}
+
+	void PostRestorePersistedProperty(const UObject* InObject, const FProperty* InProperty) const
+	{
+		const UClass* Class = InObject->GetClass();
+		while (Class)
+		{
+			if (const FPostRestorePersistedProperty* Found = ClassPostRestorePersistedProperty.Find(Class))
+			{
+				Found->ExecuteIfBound(InObject, InProperty);
+			}
+			Class = Class->GetSuperClass();
+		};
+	}
+	
+	TMap<TWeakObjectPtr<const UClass>, FShouldPersistProperty> ClassShouldPersistProperty;
+	TMap<TWeakObjectPtr<const UClass>, FPostRestorePersistedProperty> ClassPostRestorePersistedProperty;
+
+	friend class ULevelStreamingPersistenceManager;
+};
\ No newline at end of file
diff --git a/Engine/Plugins/Runtime/LevelStreamingPersistence/Source/LevelStreamingPersistence/Public/LevelStreamingPersistenceSettings.h b/Engine/Plugins/Runtime/LevelStreamingPersistence/Source/LevelStreamingPersistence/Public/LevelStreamingPersistenceSettings.h
new file mode 100644
--- /dev/null
+++ b/Engine/Plugins/Runtime/LevelStreamingPersistence/Source/LevelStreamingPersistence/Public/LevelStreamingPersistenceSettings.h
@@ -0,0 +1,29 @@
+// Copyright Epic Games, Inc. All Rights Reserved.
+
+#pragma once
+
+#include "CoreMinimal.h"
+#include "LevelStreamingPersistenceSettings.generated.h"
+
+USTRUCT()
+struct LEVELSTREAMINGPERSISTENCE_API FLevelStreamingPersistentProperty
+{
+	GENERATED_BODY()
+
+	UPROPERTY()
+	FString Path;
+
+	UPROPERTY()
+	bool bIsPublic = false;
+};
+
+UCLASS(Config = Engine)
+class LEVELSTREAMINGPERSISTENCE_API ULevelStreamingPersistenceSettings : public UObject
+{
+	GENERATED_BODY()
+
+	UPROPERTY(Config)
+	TArray<FLevelStreamingPersistentProperty> Properties;
+
+	friend class FLevelStreamingPersistentPropertiesInfo;
+};
\ No newline at end of file
diff --git a/Engine/Plugins/Runtime/LevelStreamingPersistence/Source/LevelStreamingPersistence/Public/LevelStreamingPersistentObjectPropertyBag.h b/Engine/Plugins/Runtime/LevelStreamingPersistence/Source/LevelStreamingPersistence/Public/LevelStreamingPersistentObjectPropertyBag.h
new file mode 100644
--- /dev/null
+++ b/Engine/Plugins/Runtime/LevelStreamingPersistence/Source/LevelStreamingPersistence/Public/LevelStreamingPersistentObjectPropertyBag.h
@@ -0,0 +1,89 @@
+// Copyright Epic Games, Inc. All Rights Reserved.
+
+#pragma once
+
+#include "CoreTypes.h"
+#include "UObject/UnrealType.h"
+#include "PropertyBag.h"
+
+struct FLevelStreamingPersistentObjectPropertyBag
+{
+	bool Initialize(TFunctionRef<const UPropertyBag* ()> Func);
+	bool IsValid() const { return PropertyBag.IsValid(); }
+	const UPropertyBag* GetPropertyBagStruct() const { return PropertyBag.GetPropertyBagStruct(); }
+	const FProperty* FindPropertyByName(const FName InPropertyName) const;
+	const FProperty* CopyPropertyValueFromObject(const UObject* InObject, const FProperty* InObjectProperty);
+	const FProperty* CopyPropertyValueToObject(UObject* InObject, const FProperty* InObjectProperty) const;
+	bool ComparePropertyValueWithObject(const UObject* InObject, const FProperty* InObjectProperty, bool& bOutIsIdentical) const;
+	void ForEachProperty(TFunctionRef<void(const FProperty*)> Func) const;
+	void DumpContent(TFunctionRef<void(const FProperty*, const FString&)> Func, const TArray<const FProperty*>* InDumpProperties = nullptr);
+
+	template<typename PropertyType>
+	bool SetPropertyValue(const FName InPropertyName, const PropertyType& InPropertyValue);
+
+	template<typename PropertyType>
+	bool GetPropertyValue(const FName InPropertyName, PropertyType& OutPropertyValue) const;
+
+	bool SetPropertyValueFromString(const FName InPropertyName, const FString& InPropertyValue);
+	bool GetPropertyValueAsString(const FName InPropertyName, FString& OutPropertyValueAsString) const;
+
+private:
+	template<typename PropertyType>
+	EPropertyBagResult SetValue(const FName InPropertyName, const PropertyType& InPropertyValue) { return EPropertyBagResult::PropertyNotFound; }
+	
+	template<typename PropertyType>
+	TValueOrError<PropertyType, EPropertyBagResult> GetValue(const FName InPropertyName) const { return MakeError(EPropertyBagResult::PropertyNotFound); }
+
+	template<typename PropertyType>
+	TValueOrError<PropertyType, EPropertyBagResult> GetValueStruct(const FName InPropertyName) const;
+
+	template<typename PropertyType>
+	struct TPropertyBagPropertyType { static const EPropertyBagPropertyType Value = EPropertyBagPropertyType::None; };
+
+	FInstancedPropertyBag PropertyBag;
+};
+
+template<typename PropertyType>
+TValueOrError<PropertyType, EPropertyBagResult> FLevelStreamingPersistentObjectPropertyBag::GetValueStruct(const FName InPropertyName) const
+{
+	TValueOrError<FStructView, EPropertyBagResult> Result = PropertyBag.GetValueStruct(InPropertyName);
+	if (Result.HasValue() && Result.GetValue().GetPtr<PropertyType>())
+	{
+		return MakeValue(Result.GetValue().Get<PropertyType>());
+	}
+	return MakeError(Result.HasError() ? Result.GetError() : EPropertyBagResult::PropertyNotFound);
+}
+
+template<typename PropertyType>
+bool FLevelStreamingPersistentObjectPropertyBag::SetPropertyValue(const FName InPropertyName, const PropertyType& InPropertyValue)
+{
+	if (IsValid())
+	{
+		const FPropertyBagPropertyDesc* Desc = PropertyBag.FindPropertyDescByName(InPropertyName);
+		if (Desc && (Desc->ValueType == TPropertyBagPropertyType<PropertyType>::Value))
+		{
+			EPropertyBagResult Result = SetValue(InPropertyName, InPropertyValue);
+			return Result == EPropertyBagResult::Success;
+		}
+	}
+	return false;
+}
+
+template<typename PropertyType>
+bool FLevelStreamingPersistentObjectPropertyBag::GetPropertyValue(const FName InPropertyName, PropertyType& OutPropertyValue) const
+{
+	if (IsValid())
+	{
+		const FPropertyBagPropertyDesc* Desc = PropertyBag.FindPropertyDescByName(InPropertyName);
+		if (Desc && (Desc->ValueType == TPropertyBagPropertyType<PropertyType>::Value))
+		{
+			TValueOrError<PropertyType, EPropertyBagResult> Result = GetValue<PropertyType>(InPropertyName);
+			if (Result.HasValue())
+			{
+				OutPropertyValue = Result.GetValue();
+				return true;
+			}
+		}
+	}
+	return false;
+}
\ No newline at end of file
