LevelStreamingPeristence : Added support to serialize (save/load) the LevelStreamingPersistenceManager current state into an Archive.
- Reduced overall memory footprint
- Added option to use or not tagged property serialization

diff --git a/Engine/Plugins/Runtime/LevelStreamingPersistence/Source/LevelStreamingPersistence/Private/LevelStreamingPersistenceManager.cpp b/Engine/Plugins/Runtime/LevelStreamingPersistence/Source/LevelStreamingPersistence/Private/LevelStreamingPersistenceManager.cpp
--- a/Engine/Plugins/Runtime/LevelStreamingPersistence/Source/LevelStreamingPersistence/Private/LevelStreamingPersistenceManager.cpp
+++ b/Engine/Plugins/Runtime/LevelStreamingPersistence/Source/LevelStreamingPersistence/Private/LevelStreamingPersistenceManager.cpp
@@ -17,56 +17,113 @@
 #include "Engine/Engine.h"
 #include "GameFramework/WorldSettings.h"
 
+#if !UE_BUILD_SHIPPING
+#include "HAL/FileManager.h"
+#include "Misc/Paths.h"
+#endif
+
 DEFINE_LOG_CATEGORY_STATIC(LogLevelStreamingPersistenceManager, All, All);
 
 /*
  * Achives used by ULevelStreamingPersistenceManager
  */
 
+// By default, LevelStreamingPersistenceManager uses tagged property serialization, which is more flexible as it
+// allows to load a previously serialized snapshot of ULevelStreamingPersistenceManager even if 
+// LevelStreamingPersistenceSettings Properties has changed afterwards.
+// Setting this flag to false will use less memory, with the disadvantage of not being resilient to changes
+static const bool GUseTaggedPropertySerialization = true;
+
 class FPersistentPropertiesArchive : public FObjectAndNameAsStringProxyArchive
 {
 public:
-	FPersistentPropertiesArchive(FArchive& InArchive, const FLevelStreamingPersistentPropertyArray& InCustomPropertyList)
+	FPersistentPropertiesArchive(FArchive& InArchive, const TArray<const FProperty*>& InPersistentProperties)
 		: FObjectAndNameAsStringProxyArchive(InArchive, /*bLoadIfFindFails*/false)
+		, PersistentProperties(InPersistentProperties)
 	{
 		check(InArchive.IsPersistent());
 		check(InArchive.IsFilterEditorOnly());
 		check(InArchive.ShouldSkipBulkData());
 		SetIsLoading(InArchive.IsLoading());
 		SetIsSaving(InArchive.IsSaving());
+		SetIsTextFormat(InArchive.IsTextFormat());
+		SetWantBinaryPropertySerialization(InArchive.WantBinaryPropertySerialization());
 		SetIsPersistent(true);
 		SetFilterEditorOnly(true);
 		ArShouldSkipBulkData = true;
-		ArUseCustomPropertyList = !InCustomPropertyList.IsEmpty();
-		ArCustomPropertyList = ArUseCustomPropertyList ? &InCustomPropertyList[0] : nullptr;
+
+		if (WantBinaryPropertySerialization())
+		{
+			// Custom property lists only work with binary serialization, not tagged property serialization.
+			BuildSerializedPropertyList();
+			ArUseCustomPropertyList = !CustomPropertyList.IsEmpty();
+			ArCustomPropertyList = ArUseCustomPropertyList ? &CustomPropertyList[0] : nullptr;
+		}
+	}
+
+	virtual bool ShouldSkipProperty(const FProperty* InProperty) const
+	{
+		if (FObjectAndNameAsStringProxyArchive::ShouldSkipProperty(InProperty))
+		{
+			return true;
+		}
+
+		if (!ArUseCustomPropertyList)
+		{
+			return !PersistentProperties.Contains(InProperty);
+		}
+
+		return false;
 	}
 
 	virtual FArchive& operator<<(FLazyObjectPtr& Value) override { return FArchiveUObject::SerializeLazyObjectPtr(*this, Value); }
+
+private:
+
+	void BuildSerializedPropertyList()
+	{
+		for (const FProperty* Property : PersistentProperties)
+		{
+			CustomPropertyList.Emplace(const_cast<FProperty*>(Property));
+		}
+
+		// Link changed properties
+		if (!CustomPropertyList.IsEmpty())
+		{
+			for (int i = 0; i < CustomPropertyList.Num() - 1; ++i)
+			{
+				CustomPropertyList[i].PropertyListNext = &CustomPropertyList[i + 1];
+			}
+		}
+	}
+
+	const TArray<const FProperty*>& PersistentProperties;
+	FLevelStreamingPersistentPropertyArray CustomPropertyList;
 };
 
 class FPersistentPropertiesWriter : public FMemoryWriter
 {
 public:
-	FPersistentPropertiesWriter(TArray<uint8, TSizedDefaultAllocator<32>>& InBytes)
+	FPersistentPropertiesWriter(bool bInUseTaggedPropertySerialization, TArray<uint8, TSizedDefaultAllocator<32>>& InBytes)
 		: FMemoryWriter(InBytes, true)
 	{
 		SetFilterEditorOnly(true);
 		ArShouldSkipBulkData = true;
 		SetIsTextFormat(false);
-		SetWantBinaryPropertySerialization(true);
+		SetWantBinaryPropertySerialization(!bInUseTaggedPropertySerialization);
 	}
 };
 
 class FPersistentPropertiesReader : public FMemoryReader
 {
 public:
-	FPersistentPropertiesReader(const TArray<uint8>& InBytes)
+	FPersistentPropertiesReader(bool bInUseTaggedPropertySerialization, const TArray<uint8>& InBytes)
 		: FMemoryReader(InBytes, true)
 	{
 		SetFilterEditorOnly(true);
 		ArShouldSkipBulkData = true;
 		SetIsTextFormat(false);
-		SetWantBinaryPropertySerialization(true);
+		SetWantBinaryPropertySerialization(!bInUseTaggedPropertySerialization);
 	}
 };
 
@@ -76,6 +133,7 @@ public:
 
 ULevelStreamingPersistenceManager::ULevelStreamingPersistenceManager(const FObjectInitializer& ObjectInitializer)
 	: PersistenceModule(&ILevelStreamingPersistenceModule::Get())
+	, bUseTaggedPropertySerialization(GUseTaggedPropertySerialization)
 {
 }
 
@@ -113,11 +171,9 @@ void ULevelStreamingPersistenceManager::Initialize(FSubsystemCollectionBase& Col
 
 	FLevelStreamingDelegates::OnLevelBeginMakingInvisible.AddUObject(this, &ULevelStreamingPersistenceManager::OnLevelBeginMakingInvisible);
 	FLevelStreamingDelegates::OnLevelBeginMakingVisible.AddUObject(this, &ULevelStreamingPersistenceManager::OnLevelBeginMakingVisible);
-}
 
-void ULevelStreamingPersistenceManager::PostInitialize()
-{
-	PersistentPropertiesInfo.Initialize();
+	PersistentPropertiesInfo = NewObject<ULevelStreamingPersistentPropertiesInfo>(this);
+	PersistentPropertiesInfo->Initialize();
 }
 
 void ULevelStreamingPersistenceManager::Deinitialize()
@@ -152,86 +208,60 @@ bool ULevelStreamingPersistenceManager::GetPropertyValueAsString(const FString&
 
 void ULevelStreamingPersistenceManager::OnLevelBeginMakingInvisible(UWorld* World, const ULevelStreaming* InStreamingLevel, ULevel* InLoadedLevel)
 {
-	if (IsValid(InLoadedLevel) && (World == GetWorld()))
+	if (IsValid(InLoadedLevel) && (World == GetWorld()) && IsEnabled())
 	{
 		SaveLevelPersistentPropertyValues(InLoadedLevel);
 	}
 }
 
-void ULevelStreamingPersistenceManager::OnLevelBeginMakingVisible(UWorld* World, const ULevelStreaming* InStreamingLevel, ULevel* InLoadedLevel)
-{
-	if (IsValid(InLoadedLevel) && (World == GetWorld()))
-	{
-		RestoreLevelPersistentPropertyValues(InLoadedLevel);
-	}
-}
-
-bool ULevelStreamingPersistenceManager::SaveLevelPersistentPropertyValues(const ULevel* InLevel)
+bool ULevelStreamingPersistenceManager::SaveLevelPersistentPropertyValues(const ULevel * InLevel)
 {
-	if (!IsEnabled())
-	{
-		return false;
-	}
-
 	TRACE_CPUPROFILER_EVENT_SCOPE(ULevelStreamingPersistenceManager::SaveLevelPersistentPropertyValues);
 
+	check(IsValid(InLevel) && IsEnabled());
 	const FString LevelPathName = InLevel->GetPathName();
 	FLevelStreamingPersistentPropertyValues& LevelProperties = LevelsPropertyValues.FindOrAdd(LevelPathName);
 
-	auto SavePrivateProperties = [this, &LevelProperties](const UObject* Object)
+	auto SavePrivateProperties = [this, &LevelProperties]()
 	{
 		int32 SavedCount = 0;
-		check(IsValid(Object));
-		const UClass* ObjectClass = Object->GetClass();
-		if (PersistentPropertiesInfo.HasProperties(FLevelStreamingPersistentPropertiesInfo::PropertyType_Private, ObjectClass))
+		for (auto& [ObjectPath, ObjectPrivatePropertyValues] : LevelProperties.ObjectsPrivatePropertyValues)
 		{
-			// Don't handle new objects (not part of original snapshot)
-			if (FLevelStreamingPersistentObjectPrivateProperties* ObjectPrivatePropertyValues = LevelProperties.ObjectsPrivatePropertyValues.Find(Object->GetPathName()))
+			if (const UObject* Object = FindObject<UObject>(nullptr, *ObjectPath); IsValid(Object))
 			{
-				ObjectPrivatePropertyValues->PayloadData.Reset();
-				ObjectPrivatePropertyValues->PersistentProperties.Reset();
-				if (DiffWithSnapshot(Object, ObjectPrivatePropertyValues->Snapshot, ObjectPrivatePropertyValues->PersistentProperties))
+				ObjectPrivatePropertyValues.PayloadData.Reset();
+				ObjectPrivatePropertyValues.PersistentProperties.Reset();
+				if (DiffWithSnapshot(Object, ObjectPrivatePropertyValues.Snapshot, ObjectPrivatePropertyValues.PersistentProperties))
 				{
-					FLevelStreamingPersistentPropertyArray CustomPropertyList;
-					ObjectPrivatePropertyValues->BuildSerializedPropertyList(CustomPropertyList);
-					FPersistentPropertiesWriter WriterAr(ObjectPrivatePropertyValues->PayloadData);
-					FPersistentPropertiesArchive Ar(WriterAr, CustomPropertyList);
+					FPersistentPropertiesWriter WriterAr(bUseTaggedPropertySerialization, ObjectPrivatePropertyValues.PayloadData);
+					FPersistentPropertiesArchive Ar(WriterAr, ObjectPrivatePropertyValues.PersistentProperties);
 					Object->SerializeScriptProperties(Ar);
+					SavedCount += (ObjectPrivatePropertyValues.PersistentProperties.Num() > 0) ? 1 : 0;
 				}
 			}
 		}
 		return SavedCount;
 	};
 
-	auto SavePublicProperties = [this, &LevelProperties](const UObject* Object)
+	auto SavePublicProperties = [this, &LevelProperties]()
 	{
 		int32 SavedCount = 0;
-		check(IsValid(Object));
-		const UClass* ObjectClass = Object->GetClass();
-		// Save public properties
-		if (PersistentPropertiesInfo.HasProperties(FLevelStreamingPersistentPropertiesInfo::PropertyType_Public, ObjectClass))
+		for (auto& [ObjectPath, ObjectPublicPropertyValues] : LevelProperties.ObjectsPublicPropertyValues)
 		{
-			TArray<FProperty*, TInlineAllocator<32>> SerializedObjectProperties;
-			PersistentPropertiesInfo.ForEachProperty(FLevelStreamingPersistentPropertiesInfo::PropertyType_Public, ObjectClass, [this, Object, &SerializedObjectProperties](FProperty* ObjectProperty)
-			{
-				if (!PersistenceModule->ShouldPersistProperty(Object, ObjectProperty))
-				{
-					return;
-				}
-				SerializedObjectProperties.Add(ObjectProperty);
-			});
-			if (!SerializedObjectProperties.IsEmpty())
+			if (const UObject* Object = FindObject<UObject>(nullptr, *ObjectPath); IsValid(Object))
 			{
-				FLevelStreamingPersistentObjectPublicProperties& ObjectPublicPropertyValues = LevelProperties.ObjectsPublicPropertyValues.FindOrAdd(Object->GetPathName());
-				if (ObjectPublicPropertyValues.PropertyBag.Initialize([this, ObjectClass]() { return PersistentPropertiesInfo.GetPropertyBagFromClass(FLevelStreamingPersistentPropertiesInfo::PropertyType_Public, ObjectClass); }))
+				if (ensure(ObjectPublicPropertyValues.PropertyBag.IsValid()))
 				{
 					ObjectPublicPropertyValues.PersistentProperties.Reset();
-					for (FProperty* ObjectProperty : SerializedObjectProperties)
+					for (const FProperty* BagProperty : ObjectPublicPropertyValues.PropertiesToPersist)
 					{
-						if (const FProperty* BagProperty = ObjectPublicPropertyValues.PropertyBag.CopyPropertyValueFromObject(Object, ObjectProperty))
+						if (FProperty* ObjectProperty = Object->GetClass()->FindPropertyByName(BagProperty->GetFName()))
 						{
-							ObjectPublicPropertyValues.PersistentProperties.Add(BagProperty);
-							++SavedCount;
+							if (ObjectPublicPropertyValues.PropertyBag.CopyPropertyValueFromObject(Object, ObjectProperty))
+							{
+								ObjectPublicPropertyValues.PersistentProperties.Add(BagProperty);
+								++SavedCount;
+							}
 						}
 					}
 				}
@@ -240,81 +270,179 @@ bool ULevelStreamingPersistenceManager::SaveLevelPersistentPropertyValues(const
 		return SavedCount;
 	};
 
-	int32 SavedCount = 0;
-	ForEachObjectWithOuter(InLevel, [this, &SavePrivateProperties, &SavePublicProperties, &SavedCount](UObject* Object)
-	{	
-		SavedCount += SavePrivateProperties(Object);
-		SavedCount += SavePublicProperties(Object);
-	}, true, RF_NoFlags, EInternalObjectFlags::Garbage);
+	const int32 SavedCount = SavePrivateProperties() + SavePublicProperties();
 	return SavedCount > 0;
 }
 
-bool ULevelStreamingPersistenceManager::RestoreLevelPersistentPropertyValues(ULevel* InLevel) const
+void ULevelStreamingPersistenceManager::OnLevelBeginMakingVisible(UWorld* World, const ULevelStreaming* InStreamingLevel, ULevel* InLoadedLevel)
 {
-	if (!IsEnabled())
+	if (IsValid(InLoadedLevel) && (World == GetWorld()) && IsEnabled())
 	{
-		return false;
+		RestoreLevelPersistentPropertyValues(InLoadedLevel);
 	}
+}
 
+bool ULevelStreamingPersistenceManager::RestoreLevelPersistentPropertyValues(const ULevel* InLevel) const
+{
 	TRACE_CPUPROFILER_EVENT_SCOPE(ULevelStreamingPersistenceManager::RestoreLevelPersistentPropertyValues);
 
-	auto SavePrivatePropertiesSnapshot = [this](UObject* Object, FLevelStreamingPersistentPropertyValues& LevelProperties)
+	check(IsValid(InLevel) && IsEnabled());
+
+	auto CreatePrivatePropertiesSnapshot = [this](UObject* Object, FLevelStreamingPersistentPropertyValues& LevelProperties)
 	{
-		int BuiltSnapshotCount = 0;
+		int CreatedCount = 0;
 		const UClass* ObjectClass = Object->GetClass();
-
-		if (PersistentPropertiesInfo.HasProperties(FLevelStreamingPersistentPropertiesInfo::PropertyType_Private, ObjectClass))
+		if (PersistentPropertiesInfo->HasProperties(ULevelStreamingPersistentPropertiesInfo::PropertyType_Private, ObjectClass))
 		{
-			FLevelStreamingPersistentObjectPrivateProperties& ObjectPrivatePropertyValues = LevelProperties.ObjectsPrivatePropertyValues.FindOrAdd(Object->GetPathName());
-			if (!ObjectPrivatePropertyValues.Snapshot.IsValid())
+			FLevelStreamingPersistentObjectPrivateProperties* ObjectPrivatePropertyValues = LevelProperties.ObjectsPrivatePropertyValues.Find(Object->GetPathName());
+			if (!ObjectPrivatePropertyValues || !ObjectPrivatePropertyValues->Snapshot.IsValid())
 			{
-				if (BuildSnapshot(Object, ObjectPrivatePropertyValues.Snapshot))
+				if (!ObjectPrivatePropertyValues)
 				{
-					++BuiltSnapshotCount;
+					ObjectPrivatePropertyValues = &LevelProperties.ObjectsPrivatePropertyValues.Add(Object->GetPathName());
+					ObjectPrivatePropertyValues->SourceClassPathName = ObjectClass->GetPathName();
 				}
 			}
+			
+			if (!ObjectPrivatePropertyValues->Snapshot.IsValid() && BuildSnapshot(Object, ObjectPrivatePropertyValues->Snapshot))
+			{
+				++CreatedCount;
+			}
+			
+			// Clean-up private entries
+			if (!ObjectPrivatePropertyValues->Snapshot.IsValid())
+			{
+				LevelProperties.ObjectsPrivatePropertyValues.Remove(Object->GetPathName());
+			}
 		}
-		return BuiltSnapshotCount;
+		return CreatedCount;
 	};
 
-	int32 RestoredCount = 0;
-	int32 SavedCount = 0;
+	auto CreatePublicPropertiesEntry = [this](const UObject* Object, FLevelStreamingPersistentPropertyValues& LevelProperties)
+	{
+		int32 CreatedCount = 0;
+		check(IsValid(Object));
+		const UClass* ObjectClass = Object->GetClass();
+
+		if (PersistentPropertiesInfo->HasProperties(ULevelStreamingPersistentPropertiesInfo::PropertyType_Public, ObjectClass))
+		{
+			FLevelStreamingPersistentObjectPublicProperties* ObjectPublicPropertyValues = LevelProperties.ObjectsPublicPropertyValues.Find(Object->GetPathName());
+
+			// Build PropertiesToPersist if necessary
+			TArray<FProperty*, TInlineAllocator<32>> SerializedObjectProperties;
+			if (!ObjectPublicPropertyValues || ObjectPublicPropertyValues->PropertiesToPersist.IsEmpty())
+			{
+				PersistentPropertiesInfo->ForEachProperty(ULevelStreamingPersistentPropertiesInfo::PropertyType_Public, ObjectClass, [this, Object, &SerializedObjectProperties](FProperty* ObjectProperty)
+				{
+					if (PersistenceModule->ShouldPersistProperty(Object, ObjectProperty))
+					{
+						SerializedObjectProperties.Add(ObjectProperty);
+					}
+				});
+
+				if (!SerializedObjectProperties.IsEmpty())
+				{
+					check(!ObjectPublicPropertyValues || ObjectPublicPropertyValues->PropertyBag.IsValid());
+					if (!ObjectPublicPropertyValues)
+					{
+						// Create entry if necessary and initialize PropertyBag
+						ObjectPublicPropertyValues = &LevelProperties.ObjectsPublicPropertyValues.Add(Object->GetPathName());
+						ObjectPublicPropertyValues->SourceClassPathName = ObjectClass->GetPathName();
+						ObjectPublicPropertyValues->PropertyBag.Initialize([this, ObjectClass]() { return PersistentPropertiesInfo->GetPropertyBagFromClass(ULevelStreamingPersistentPropertiesInfo::PropertyType_Public, ObjectClass); });
+					}
+					if (ObjectPublicPropertyValues->PropertyBag.IsValid())
+					{
+						check(ObjectPublicPropertyValues->PropertiesToPersist.IsEmpty());
+						for (FProperty* ObjectProperty : SerializedObjectProperties)
+						{
+							if (const FProperty* BagProperty = ObjectPublicPropertyValues->PropertyBag.GetCompatibleProperty(ObjectProperty))
+							{
+								ObjectPublicPropertyValues->PropertiesToPersist.Add(BagProperty);
+								++CreatedCount;
+							}
+						}
+					}
+				}
+			}
+
+			// Clean-up public entries
+			if (ObjectPublicPropertyValues && (!ObjectPublicPropertyValues->PropertyBag.IsValid() || ObjectPublicPropertyValues->PropertiesToPersist.IsEmpty()))
+			{
+				LevelProperties.ObjectsPublicPropertyValues.Remove(Object->GetPathName());
+			}
+		}
+		return CreatedCount;
+	};
 
+	int32 CreatedCount = 0;
 	const FString LevelPathName = InLevel->GetPathName();
 	FLevelStreamingPersistentPropertyValues* LevelProperties = LevelsPropertyValues.Find(LevelPathName);
+	const bool bIsNewLevelEntry = (LevelProperties == nullptr);
+
+	// The first time, create a snapshot of persistent properties
+	if (bIsNewLevelEntry || !LevelProperties->bIsMakingVisibleCacheValid)
 	{
-		// The first time, create a snapshot of persistent properties
-		if (!LevelProperties)
-		{
-			TRACE_CPUPROFILER_EVENT_SCOPE(ULevelStreamingPersistenceManager::SavePrivatePropertiesSnapshot);
+		TRACE_CPUPROFILER_EVENT_SCOPE(ULevelStreamingPersistenceManager::CreateLevelPropertiesEntries);
+
+		FLevelStreamingPersistentPropertyValues& NewLevelProperties = LevelsPropertyValues.FindOrAdd(LevelPathName);
+		LevelProperties = &NewLevelProperties;
 
-			FLevelStreamingPersistentPropertyValues& NewLevelProperties = LevelsPropertyValues.FindOrAdd(LevelPathName);
-			ForEachObjectWithOuter(InLevel, [this, &NewLevelProperties, &SavedCount, &SavePrivatePropertiesSnapshot](UObject* Object)
+		TSet<FString> ValidObjects;
+		ForEachObjectWithOuter(InLevel, [this, bIsNewLevelEntry, &NewLevelProperties, &ValidObjects, &CreatedCount, &CreatePrivatePropertiesSnapshot, &CreatePublicPropertiesEntry](UObject* Object)
+		{
+			int32 EntryCreatedCount = CreatePrivatePropertiesSnapshot(Object, NewLevelProperties);
+			EntryCreatedCount += CreatePublicPropertiesEntry(Object, NewLevelProperties);
+			if (!bIsNewLevelEntry && (EntryCreatedCount > 0))
 			{
-				SavedCount += SavePrivatePropertiesSnapshot(Object, NewLevelProperties);
-			}, true, RF_NoFlags, EInternalObjectFlags::Garbage);
-		}
-		// Else, restore persistent properties
-		else
+				ValidObjects.Add(Object->GetPathName());
+			}
+			CreatedCount += EntryCreatedCount;
+		}, true, RF_NoFlags, EInternalObjectFlags::Garbage);
+
+		// Clean-up object entries that don't exist in the level
+		if (ValidObjects.Num() > 0)
 		{
-			for (const auto& [ObjectPath, ObjectPrivatePropertyValues] : LevelProperties->ObjectsPrivatePropertyValues)
+			for (auto It = LevelProperties->ObjectsPrivatePropertyValues.CreateIterator(); It; ++It)
 			{
-				if (UObject* Object = FindObject<UObject>(nullptr, *ObjectPath))
+				if (!ValidObjects.Contains(It.Key()))
 				{
-					RestoredCount += RestorePrivateProperties(Object, ObjectPrivatePropertyValues);
+					It.RemoveCurrent();
 				}
 			}
-			for (const auto& [ObjectPath, ObjectPublicPropertyValues] : LevelProperties->ObjectsPublicPropertyValues)
+			for (auto It = LevelProperties->ObjectsPublicPropertyValues.CreateIterator(); It; ++It)
 			{
-				if (UObject* Object = FindObject<UObject>(nullptr, *ObjectPath))
+				if (!ValidObjects.Contains(It.Key()))
 				{
-					RestoredCount += RestorePublicProperties(Object, ObjectPublicPropertyValues);
+					It.RemoveCurrent();
 				}
 			}
 		}
+
+		// Make level entry's as valid
+		LevelProperties->bIsMakingVisibleCacheValid = true;
+	}
+
+	// Restore persistent properties
+	int32 RestoredCount = 0;
+	if (LevelProperties)
+	{
+		for (const auto& [ObjectPath, ObjectPrivatePropertyValues] : LevelProperties->ObjectsPrivatePropertyValues)
+		{
+			if (UObject* Object = FindObject<UObject>(nullptr, *ObjectPath))
+			{
+				RestoredCount += RestorePrivateProperties(Object, ObjectPrivatePropertyValues);
+			}
+		}
+		for (const auto& [ObjectPath, ObjectPublicPropertyValues] : LevelProperties->ObjectsPublicPropertyValues)
+		{
+			if (UObject* Object = FindObject<UObject>(nullptr, *ObjectPath))
+			{
+				RestoredCount += RestorePublicProperties(Object, ObjectPublicPropertyValues);
+			}
+		}
 	}
 	
-	return (RestoredCount > 0) || (SavedCount > 0);
+	return ((RestoredCount + CreatedCount) > 0);
 }
 
 int32 ULevelStreamingPersistenceManager::RestorePrivateProperties(UObject* Object, const FLevelStreamingPersistentObjectPrivateProperties& PersistentObjectProperties) const
@@ -325,10 +453,8 @@ int32 ULevelStreamingPersistenceManager::RestorePrivateProperties(UObject* Objec
 	// Restore private properties
 	if (!PersistentObjectProperties.PayloadData.IsEmpty())
 	{
-		FLevelStreamingPersistentPropertyArray CustomPropertyList;
-		PersistentObjectProperties.BuildSerializedPropertyList(CustomPropertyList);
-		FPersistentPropertiesReader ReaderAr(PersistentObjectProperties.PayloadData);
-		FPersistentPropertiesArchive Ar(ReaderAr, CustomPropertyList);
+		FPersistentPropertiesReader ReaderAr(bUseTaggedPropertySerialization, PersistentObjectProperties.PayloadData);
+		FPersistentPropertiesArchive Ar(ReaderAr, PersistentObjectProperties.PersistentProperties);
 		Object->SerializeScriptProperties(Ar);
 		PersistentObjectProperties.ForEachPersistentProperty([this, Object](const FProperty* ObjectProperty)
 		{
@@ -432,11 +558,11 @@ bool ULevelStreamingPersistenceManager::BuildSnapshot(const UObject* InObject, F
 {
 	int SavedCount = 0;
 	const UClass* ObjectClass = InObject->GetClass();
-	if (PersistentPropertiesInfo.HasProperties(FLevelStreamingPersistentPropertiesInfo::PropertyType_Private, ObjectClass))
+	if (PersistentPropertiesInfo->HasProperties(ULevelStreamingPersistentPropertiesInfo::PropertyType_Private, ObjectClass))
 	{
-		if (OutSnapshot.Initialize([this, ObjectClass]() { return PersistentPropertiesInfo.GetPropertyBagFromClass(FLevelStreamingPersistentPropertiesInfo::PropertyType_Private, ObjectClass); }))
+		if (OutSnapshot.Initialize([this, ObjectClass]() { return PersistentPropertiesInfo->GetPropertyBagFromClass(ULevelStreamingPersistentPropertiesInfo::PropertyType_Private, ObjectClass); }))
 		{
-			PersistentPropertiesInfo.ForEachProperty(FLevelStreamingPersistentPropertiesInfo::PropertyType_Private, ObjectClass, [this, InObject, &SavedCount, &OutSnapshot](FProperty* ObjectProperty)
+			PersistentPropertiesInfo->ForEachProperty(ULevelStreamingPersistentPropertiesInfo::PropertyType_Private, ObjectClass, [this, InObject, &SavedCount, &OutSnapshot](FProperty* ObjectProperty)
 			{
 				if (OutSnapshot.CopyPropertyValueFromObject(InObject, ObjectProperty))
 				{
@@ -456,16 +582,15 @@ bool ULevelStreamingPersistenceManager::DiffWithSnapshot(const UObject* InObject
 		const UClass* ObjectClass = InObject->GetClass();
 
 		// Find changed properties
-		PersistentPropertiesInfo.ForEachProperty(FLevelStreamingPersistentPropertiesInfo::PropertyType_Private, ObjectClass, [this, InObject, &InSnapshot, &OutChangedProperties](FProperty* ObjectProperty)
+		PersistentPropertiesInfo->ForEachProperty(ULevelStreamingPersistentPropertiesInfo::PropertyType_Private, ObjectClass, [this, InObject, &InSnapshot, &OutChangedProperties](FProperty* ObjectProperty)
 		{
-			if (!PersistenceModule->ShouldPersistProperty(InObject, ObjectProperty))
-			{
-				return;
-			}
-			bool bIsIdentical = true;
-			if (InSnapshot.ComparePropertyValueWithObject(InObject, ObjectProperty, bIsIdentical) && !bIsIdentical)
+			if (PersistenceModule->ShouldPersistProperty(InObject, ObjectProperty))
 			{
-				OutChangedProperties.Add(ObjectProperty);
+				bool bIsIdentical = true;
+				if (InSnapshot.ComparePropertyValueWithObject(InObject, ObjectProperty, bIsIdentical) && !bIsIdentical)
+				{
+					OutChangedProperties.Add(ObjectProperty);
+				}
 			}
 		});
 
@@ -475,6 +600,112 @@ bool ULevelStreamingPersistenceManager::DiffWithSnapshot(const UObject* InObject
 	return false;
 }
 
+bool ULevelStreamingPersistenceManager::SerializeTo(TArray<uint8>& OutPayload)
+{
+	// Serialize to archive and gather custom versions
+	TArray<uint8> PayloadData;
+	FMemoryWriter PayloadAr(PayloadData, true);
+	if (SerializeManager(PayloadAr))
+	{
+		// Serialize custom versions
+		TArray<uint8> HeaderData;
+		FMemoryWriter HeaderAr(HeaderData);
+		FCustomVersionContainer CustomVersions = PayloadAr.GetCustomVersions();
+		CustomVersions.Serialize(HeaderAr);
+
+		// Append data
+		OutPayload = MoveTemp(HeaderData);
+		OutPayload.Append(PayloadData);
+		return true;
+	}
+	return false;
+}
+
+bool ULevelStreamingPersistenceManager::InitializeFrom(const TArray<uint8>& InPayload)
+{
+	FMemoryReader PayloadAr(InPayload, true);
+
+	// Serialize custom versions
+	FCustomVersionContainer CustomVersions;
+	CustomVersions.Serialize(PayloadAr);
+	PayloadAr.SetCustomVersions(CustomVersions);
+
+	// Serialize payload
+	return SerializeManager(PayloadAr);
+}
+
+bool ULevelStreamingPersistenceManager::SerializeManager(FArchive& InArchive)
+{
+	if (!InArchive.IsPersistent())
+	{
+		UE_LOG(LogLevelStreamingPersistenceManager, Warning, TEXT("Archive must be persistent to serialize LevelStreamingPersistenceManager"));
+		return false;
+	}
+
+	FObjectAndNameAsStringProxyArchive Ar(InArchive, false);
+	Ar.SetIsLoading(InArchive.IsLoading());
+	Ar.SetIsSaving(InArchive.IsSaving());
+	Ar.SetIsPersistent(true);
+	Ar.SetFilterEditorOnly(true);
+	Ar.SetIsTextFormat(false);
+	Ar.SetWantBinaryPropertySerialization(!bUseTaggedPropertySerialization);
+	Ar.ArShouldSkipBulkData = true;
+
+	bool bLocalUseTaggedPropertySerialization = bUseTaggedPropertySerialization;
+	Ar << bLocalUseTaggedPropertySerialization;
+
+	if (bLocalUseTaggedPropertySerialization != bUseTaggedPropertySerialization)
+	{
+		UE_LOG(LogLevelStreamingPersistenceManager, Error, TEXT("Tagged property serialization mismatch : Can't serialize LevelStreamingPersistenceManager"));
+		return false;
+	}
+
+	if (Ar.IsLoading())
+	{
+		TMap<FString, FLevelStreamingPersistentPropertyValues> LocalLevelsPropertyValues;
+		Ar << LocalLevelsPropertyValues;
+
+		for (auto& [LocalLevelPathName, LocalLevelPropertyValues] : LocalLevelsPropertyValues)
+		{
+			FLevelStreamingPersistentPropertyValues& LevelProperties = LevelsPropertyValues.FindOrAdd(LocalLevelPathName);
+			for (auto& [ObjectPath, LocalPrivatePropertyValues] : LocalLevelPropertyValues.ObjectsPrivatePropertyValues)
+			{
+				if (LocalPrivatePropertyValues.Sanitize(*this))
+				{
+					check(!LocalPrivatePropertyValues.Snapshot.IsValid());
+					// Override existing entry with the loaded/sanitized version 
+					LevelProperties.ObjectsPrivatePropertyValues.Add(ObjectPath, MoveTemp(LocalPrivatePropertyValues));
+					LevelProperties.bIsMakingVisibleCacheValid = false;
+				}
+			}
+
+			for (auto& [ObjectPath, LocalPublicPropertyValues] : LocalLevelPropertyValues.ObjectsPublicPropertyValues)
+			{
+				if (LocalPublicPropertyValues.Sanitize(*this))
+				{
+					check(LocalPublicPropertyValues.PropertiesToPersist.IsEmpty());
+					// Override existing entry with the loaded/sanitized version 
+					LevelProperties.ObjectsPublicPropertyValues.Add(ObjectPath, MoveTemp(LocalPublicPropertyValues));
+					LevelProperties.bIsMakingVisibleCacheValid = false;
+				}
+			}
+
+			// If level is loaded, run restore logic right away
+			if (ULevel* Level = FindObject<ULevel>(nullptr, *LocalLevelPathName))
+			{
+				RestoreLevelPersistentPropertyValues(Level);
+			}
+		}
+	}
+	else if (Ar.IsSaving())
+	{
+		Ar << LevelsPropertyValues;
+	}
+
+	return true;
+}
+
+#if !UE_BUILD_SHIPPING
 void ULevelStreamingPersistenceManager::DumpContent() const
 {
 	UE_LOG(LogLevelStreamingPersistenceManager, Log, TEXT("World Persistence Content for %s"), *GetWorld()->GetName());
@@ -499,7 +730,7 @@ void ULevelStreamingPersistenceManager::DumpContent() const
 			{
 				if (!ObjectPropertyValues.PayloadData.IsEmpty())
 				{
-					if (const UClass* Class = PersistentPropertiesInfo.GetClassFromPropertyBag(FLevelStreamingPersistentPropertiesInfo::PropertyType_Private, ObjectPropertyValues.Snapshot.GetPropertyBagStruct()))
+					if (const UClass* Class = ObjectPropertyValues.GetSourceClass())
 					{
 						if (UObject* TempObject = GetTemporaryObjectForClass(Class))
 						{
@@ -547,11 +778,20 @@ void ULevelStreamingPersistenceManager::DumpContent() const
 		}
 	}
 }
+#endif
 
 /*
  * FLevelStreamingPersistentObjectPrivateProperties implementation
  */
 
+FLevelStreamingPersistentObjectPrivateProperties::FLevelStreamingPersistentObjectPrivateProperties(FLevelStreamingPersistentObjectPrivateProperties&& InOther)
+	: Snapshot(MoveTemp(InOther.Snapshot))
+	, SourceClassPathName(MoveTemp(InOther.SourceClassPathName))
+	, PayloadData(MoveTemp(InOther.PayloadData))
+	, PersistentProperties(InOther.PersistentProperties)
+{
+}
+
 void FLevelStreamingPersistentObjectPrivateProperties::ForEachPersistentProperty(TFunctionRef<void(const FProperty*)> Func) const
 {
 	for (const FProperty* Property : PersistentProperties)
@@ -560,27 +800,89 @@ void FLevelStreamingPersistentObjectPrivateProperties::ForEachPersistentProperty
 	}
 }
 
-void FLevelStreamingPersistentObjectPrivateProperties::BuildSerializedPropertyList(FLevelStreamingPersistentPropertyArray& OutCustomPropertyList) const
+FArchive& operator<<(FArchive& Ar, FLevelStreamingPersistentObjectPrivateProperties& PrivateProperties)
+{
+	PrivateProperties.Serialize(Ar);
+	return Ar;
+}
+
+void FLevelStreamingPersistentObjectPrivateProperties::Serialize(FArchive& Ar)
 {
-	ForEachPersistentProperty([&OutCustomPropertyList](const FProperty* Property)
+	Ar << SourceClassPathName;
+	Ar << PayloadData;
+
+	uint32 PersistentPropertiesCount = PersistentProperties.Num();
+	Ar << PersistentPropertiesCount;
+
+	const UClass* Class = GetSourceClass();
+	UE_CLOG(!Class, LogLevelStreamingPersistenceManager, Verbose, TEXT("FLevelStreamingPersistentObjectPrivateProperties : Could not find class %s"), *SourceClassPathName);
+
+	if (Ar.IsLoading())
+	{
+		PersistentProperties.Reserve(PersistentPropertiesCount);
+		for (uint32 i = 0; i < PersistentPropertiesCount; ++i)
+		{
+			FString PropertyName;
+			Ar << PropertyName;
+
+			if (const FProperty* ObjectProperty = Class ? Class->FindPropertyByName(FName(PropertyName)) : nullptr)
+			{
+				PersistentProperties.Add(ObjectProperty);
+			}
+			else
+			{
+				UE_CLOG(Class, LogLevelStreamingPersistenceManager, Verbose, TEXT("FLevelStreamingPersistentObjectPrivateProperties : Could not find private property %s"), *PropertyName);
+			}
+		}
+	}
+	else if (Ar.IsSaving())
 	{
-		OutCustomPropertyList.Emplace(const_cast<FProperty*>(Property));
-	});
+		for (const FProperty* Property : PersistentProperties)
+		{
+			FString PropertyName = Property->GetName();
+			Ar << PropertyName;
+		}
+	}
+}
+
+bool FLevelStreamingPersistentObjectPrivateProperties::Sanitize(const ULevelStreamingPersistenceManager& InManager)
+{
+	check(PersistentProperties.IsEmpty() == PayloadData.IsEmpty());
 
-	// Link changed properties
-	if (!OutCustomPropertyList.IsEmpty())
+	if (!PersistentProperties.IsEmpty())
 	{
-		for (int i = 0; i < OutCustomPropertyList.Num() - 1; ++i)
+		TArray<const FProperty*> LocalPersistentProperties = MoveTemp(PersistentProperties);
+		PersistentProperties.Reserve(LocalPersistentProperties.Num());
+		for (const FProperty* PersistentProperty : LocalPersistentProperties)
 		{
-			OutCustomPropertyList[i].PropertyListNext = &OutCustomPropertyList[i + 1];
+			if (InManager.PersistentPropertiesInfo->HasProperty(ULevelStreamingPersistentPropertiesInfo::PropertyType_Private, PersistentProperty))
+			{
+				PersistentProperties.Add(PersistentProperty);
+			}
+			else if (!InManager.bUseTaggedPropertySerialization)
+			{
+				UE_LOG(LogLevelStreamingPersistenceManager, Verbose, TEXT("FLevelStreamingPersistentObjectPrivateProperties : ULevelStreamingPersistenceManager doesn't use tagged property serialization and found an invalid persistent property %s. All private properties for this object will be ignored."), *PersistentProperty->GetName());
+				PersistentProperties.Reset();
+				return false;
+			}
 		}
 	}
+	return !PersistentProperties.IsEmpty();
 }
 
 /*
  * FLevelStreamingPersistentObjectPublicProperties implementation
  */
 
+FLevelStreamingPersistentObjectPublicProperties::FLevelStreamingPersistentObjectPublicProperties(FLevelStreamingPersistentObjectPublicProperties&& InOther)
+	: PropertiesToPersist(MoveTemp(InOther.PropertiesToPersist))
+	, SourceClassPathName(MoveTemp(InOther.SourceClassPathName))
+	, PropertyBag(MoveTemp(InOther.PropertyBag))
+	, PersistentProperties(InOther.PersistentProperties)
+
+{
+}
+
 void FLevelStreamingPersistentObjectPublicProperties::ForEachPersistentProperty(TFunctionRef<void(const FProperty*)> Func) const
 {
 	for (const FProperty* Property : PersistentProperties)
@@ -589,15 +891,123 @@ void FLevelStreamingPersistentObjectPublicProperties::ForEachPersistentProperty(
 	}
 }
 
+FArchive& operator<<(FArchive& Ar, FLevelStreamingPersistentObjectPublicProperties& PublicProperties)
+{
+	PublicProperties.Serialize(Ar);
+	return Ar;
+}
+
+void FLevelStreamingPersistentObjectPublicProperties::Serialize(FArchive& Ar)
+{
+	Ar << SourceClassPathName;
+	Ar << PropertyBag;
+
+	uint32 PersistentPropertiesCount = PersistentProperties.Num();
+	Ar << PersistentPropertiesCount;
+	
+	if (Ar.IsLoading())
+	{
+		PersistentProperties.Reserve(PersistentPropertiesCount);
+		for (uint32 i = 0; i < PersistentPropertiesCount; ++i)
+		{
+			FString PropertyName;
+			Ar << PropertyName;
+			if (const FProperty* Property = PropertyBag.FindPropertyByName(FName(PropertyName)))
+			{
+				PersistentProperties.Add(Property);
+			}
+			else
+			{
+				UE_LOG(LogLevelStreamingPersistenceManager, Verbose, TEXT("FLevelStreamingPersistentObjectPublicProperties : Could not find public property %s in property bag."), *PropertyName);
+			}
+		}
+	}
+	else if (Ar.IsSaving())
+	{
+		for (const FProperty* Property : PersistentProperties)
+		{
+			FString PropertyName = Property->GetName();
+			Ar << PropertyName;
+		}
+	}
+}
+
+bool FLevelStreamingPersistentObjectPublicProperties::Sanitize(const ULevelStreamingPersistenceManager& InManager)
+{
+	if (!PropertyBag.IsValid())
+	{
+		return false;
+	}
+
+	const UClass* SourceClass = GetSourceClass();
+	const UPropertyBag* DefaultClass = SourceClass ? InManager.PersistentPropertiesInfo->GetPropertyBagFromClass(ULevelStreamingPersistentPropertiesInfo::PropertyType_Public, SourceClass) : nullptr;
+	if (!DefaultClass)
+	{
+		return false;
+	}
+
+	FLevelStreamingPersistentObjectPropertyBag LocalPropertyBag = MoveTemp(PropertyBag);
+	TArray<const FProperty*> LocalPersistentProperties = MoveTemp(PersistentProperties);
+	PropertyBag.Initialize([&InManager, SourceClass]() { return InManager.PersistentPropertiesInfo->GetPropertyBagFromClass(ULevelStreamingPersistentPropertiesInfo::PropertyType_Public, SourceClass); });
+
+	for (const FProperty* LocalProperty : LocalPersistentProperties)
+	{
+		if (const FProperty* Property = PropertyBag.FindPropertyByName(LocalProperty->GetFName()))
+		{
+			if (PropertyBag.CopyPropertyValueFromPropertyBag(LocalPropertyBag, LocalProperty))
+			{
+				PersistentProperties.Add(Property);
+			}
+		}
+	}
+
+	if (PersistentProperties.IsEmpty())
+	{
+		return false;
+	}
+
+	return true;
+}
+
 /*
- * FLevelStreamingPersistentPropertiesInfo implementation
+ * FLevelStreamingPersistentPropertyValues implementation
  */
 
-void FLevelStreamingPersistentPropertiesInfo::Initialize()
+FArchive& operator<<(FArchive& Ar, FLevelStreamingPersistentPropertyValues& Properties)
+{
+	Properties.Serialize(Ar);
+	return Ar;
+}
+
+void FLevelStreamingPersistentPropertyValues::Serialize(FArchive& Ar)
+{
+	Ar << ObjectsPrivatePropertyValues;
+	Ar << ObjectsPublicPropertyValues;
+}
+
+/*
+ * ULevelStreamingPersistentPropertiesInfo implementation
+ */
+
+void ULevelStreamingPersistentPropertiesInfo::AddReferencedObjects(UObject* InThis, FReferenceCollector& Collector)
+{
+	ULevelStreamingPersistentPropertiesInfo* This = CastChecked<ULevelStreamingPersistentPropertiesInfo>(InThis);
+	for (int32 AccessSpecifierIndex = 0; AccessSpecifierIndex < PropertyType_Count; ++AccessSpecifierIndex)
+	{
+		for (auto& It : This->ClassesProperties[AccessSpecifierIndex])
+		{
+			Collector.AddReferencedObject(It.Key);
+		}
+	}
+	Super::AddReferencedObjects(InThis, Collector);
+}
+
+void ULevelStreamingPersistentPropertiesInfo::Initialize()
 {
 	auto FillClassesProperties = [this](const FString& InPropertyPath, bool bIsPublic)
 	{
-		TMap<TWeakObjectPtr<const UClass>, TSet<FProperty*>>& OutClassesProperties = ClassesProperties[bIsPublic ? PropertyType_Public : PropertyType_Private];
+		TMap<const UClass*, TSet<FProperty*>>& OutClassesProperties = ClassesProperties[bIsPublic ? PropertyType_Public : PropertyType_Private];
+		FSoftObjectPath(InPropertyPath).TryLoad();
 		if (FProperty* Property = TFieldPath<FProperty>(*InPropertyPath).Get())
 		{
 			if (const UClass* Class = Property->GetOwnerClass())
@@ -605,26 +1015,37 @@ void FLevelStreamingPersistentPropertiesInfo::Initialize()
 				OutClassesProperties.FindOrAdd(Class).Add(Property);
 			}
 		}
+		else
+		{
+			UE_LOG(LogLevelStreamingPersistenceManager, Log, TEXT("Could not resolve property path %s."), *InPropertyPath);
+		}
 	};
 
-	auto CreateClassDefaultPropertyBag = [this](EPropertyType InAccessSpecifier)
+	auto CreateClassDefaultPropertyBag = [this]()
 	{
-		TMap<TWeakObjectPtr<const UClass>, FInstancedPropertyBag>& ClassesDefaults = ObjectClassToPropertyBag[InAccessSpecifier];
-		for (auto& [Class, ClassProperties] : ClassesProperties[InAccessSpecifier])
+		for (int32 AccessSpecifierIndex = 0; AccessSpecifierIndex < PropertyType_Count; ++AccessSpecifierIndex)
 		{
-			check(Class.IsValid());
-			check(!ClassesDefaults.Contains(Class));
-
-			TArray<FPropertyBagPropertyDesc> Descs;
-			ForEachProperty(InAccessSpecifier, Class.Get(), [&Descs](FProperty* Property)
+			EPropertyType AccessSpecifier = EPropertyType(AccessSpecifierIndex);
+			TMap<const UClass*, FInstancedPropertyBag>& ClassesDefaults = ObjectClassToPropertyBag[AccessSpecifier];
+			for (auto& [Class, ClassProperties] : ClassesProperties[AccessSpecifier])
 			{
-				Descs.Emplace(Property->GetFName(), Property);
-			});
+				check(Class);
+				check(!ClassesDefaults.Contains(Class));
 
-			FInstancedPropertyBag& PropertyBag = ClassesDefaults.FindOrAdd(Class);
-			PropertyBag.AddProperties(Descs);
-			check(PropertyBag.GetPropertyBagStruct());
-			PropertyBagToObjectClass[InAccessSpecifier].Add(PropertyBag.GetPropertyBagStruct(), Class);
+				TArray<FPropertyBagPropertyDesc> Descs;
+				ForEachProperty(AccessSpecifier, Class, [&Descs](FProperty* Property)
+				{
+					Descs.Emplace(Property->GetFName(), Property);
+#if WITH_EDITORONLY_DATA
+					// Get rid of metadata to save memory
+					Descs.Last().MetaData.Empty();
+#endif
+				});
+
+				FInstancedPropertyBag& PropertyBag = ClassesDefaults.FindOrAdd(Class);
+				PropertyBag.AddProperties(Descs);
+				check(PropertyBag.GetPropertyBagStruct());
+			}
 		}
 	};
 
@@ -632,11 +1053,11 @@ void FLevelStreamingPersistentPropertiesInfo::Initialize()
 	{
 		FillClassesProperties(PersistentProperty.Path, PersistentProperty.bIsPublic);
 	}
-	CreateClassDefaultPropertyBag(PropertyType_Private);
-	CreateClassDefaultPropertyBag(PropertyType_Public);
+
+	CreateClassDefaultPropertyBag();
 }
 
-const UPropertyBag* FLevelStreamingPersistentPropertiesInfo::GetPropertyBagFromClass(EPropertyType InAccessSpecifier, const UClass* InClass) const
+const UPropertyBag* ULevelStreamingPersistentPropertiesInfo::GetPropertyBagFromClass(EPropertyType InAccessSpecifier, const UClass* InClass) const
 {
 	if (ensure(InAccessSpecifier < PropertyType_Count))
 	{
@@ -653,19 +1074,7 @@ const UPropertyBag* FLevelStreamingPersistentPropertiesInfo::GetPropertyBagFromC
 	return nullptr;
 }
 
-const UClass* FLevelStreamingPersistentPropertiesInfo::GetClassFromPropertyBag(EPropertyType InAccessSpecifier, const UPropertyBag* InPropertyBag) const
-{
-	if (ensure(InAccessSpecifier < PropertyType_Count))
-	{
-		if (const TWeakObjectPtr<const UClass>* FoundClass = PropertyBagToObjectClass[InAccessSpecifier].Find(InPropertyBag))
-		{
-			return FoundClass->Get();
-		}
-	}
-	return nullptr;
-}
-
-bool FLevelStreamingPersistentPropertiesInfo::HasProperty(EPropertyType InAccessSpecifier, const FProperty* InProperty) const
+bool ULevelStreamingPersistentPropertiesInfo::HasProperty(EPropertyType InAccessSpecifier, const FProperty* InProperty) const
 {
 	if (ensure(InAccessSpecifier < PropertyType_Count))
 	{
@@ -680,7 +1089,7 @@ bool FLevelStreamingPersistentPropertiesInfo::HasProperty(EPropertyType InAccess
 	return false;
 }
 
-void FLevelStreamingPersistentPropertiesInfo::ForEachProperty(EPropertyType InAccessSpecifier, const UClass* InClass, TFunctionRef<void(FProperty*)> Func) const
+void ULevelStreamingPersistentPropertiesInfo::ForEachProperty(EPropertyType InAccessSpecifier, const UClass* InClass, TFunctionRef<void(FProperty*)> Func) const
 {
 	if (ensure(InAccessSpecifier < PropertyType_Count))
 	{
@@ -699,7 +1108,7 @@ void FLevelStreamingPersistentPropertiesInfo::ForEachProperty(EPropertyType InAc
 	}
 }
 
-bool FLevelStreamingPersistentPropertiesInfo::HasProperties(EPropertyType InAccessSpecifier, const UClass* InClass) const
+bool ULevelStreamingPersistentPropertiesInfo::HasProperties(EPropertyType InAccessSpecifier, const UClass* InClass) const
 {
 	if (ensure(InAccessSpecifier < PropertyType_Count))
 	{
@@ -716,6 +1125,8 @@ bool FLevelStreamingPersistentPropertiesInfo::HasProperties(EPropertyType InAcce
 	return false;
 }
 
+#if !UE_BUILD_SHIPPING
+
 /*
  * Level Streaming Persistence Manager Console command helper
  */
@@ -788,21 +1199,13 @@ namespace LSPMConsoleCommandHelper
 	template<> FString GetValueToString(const FColor& InPropertyValue) { return InPropertyValue.ToString(); }
 	template<> FString GetValueToString(const FLinearColor& InPropertyValue) { return InPropertyValue.ToString(); }
 
-} // namespace LSPMConsoleCommandHelper
-
-bool ULevelStreamingPersistenceManager::bIsEnabled = true;
-FAutoConsoleVariableRef ULevelStreamingPersistenceManager::EnableCommand(
-	TEXT("s.LevelStreamingPersistence.Enabled"),
-	ULevelStreamingPersistenceManager::bIsEnabled,
-	TEXT("Turn on/off to enable/disable world persistent subsystem."),
-	ECVF_Default);
-
-#if !UE_BUILD_SHIPPING
+	static FString GetFilename(UWorld* InWorld)
+	{
+		FString Filename = FPaths::ProjectSavedDir() / TEXT("LevelStreamingPersistence") / FString::Printf(TEXT("LevelStreamingPersistence-%s.bin"), *InWorld->GetName());
+		return Filename;
+	}
 
-FAutoConsoleCommand ULevelStreamingPersistenceManager::DumpContentCommand(
-	TEXT("s.LevelStreamingPersistence.DumpContent"),
-	TEXT("Dump the content of WorldPersistentSubsystems"),
-	FConsoleCommandWithArgsDelegate::CreateLambda([](const TArray<FString>& Args)
+	void ForEachLevelStreamingPersistenceManager(TFunctionRef<bool(ULevelStreamingPersistenceManager*)> Func)
 	{
 		for (const FWorldContext& Context : GEngine->GetWorldContexts())
 		{
@@ -811,15 +1214,32 @@ FAutoConsoleCommand ULevelStreamingPersistenceManager::DumpContentCommand(
 			{
 				if (ULevelStreamingPersistenceManager* LevelStreamingPersistenceManager = World->GetSubsystem<ULevelStreamingPersistenceManager>())
 				{
-					LevelStreamingPersistenceManager->DumpContent();
+					if (!Func(LevelStreamingPersistenceManager))
+					{
+						break;
+					}
 				}
 			}
 		}
+	}
+
+} // namespace LSPMConsoleCommandHelper
+
+FAutoConsoleCommand ULevelStreamingPersistenceManager::DumpContentCommand(
+	TEXT("s.LevelStreamingPersistence.Debug.DumpContent"),
+	TEXT("Dump the content of LevelStreamingPersistenceManager"),
+	FConsoleCommandWithArgsDelegate::CreateLambda([](const TArray<FString>& Args)
+	{
+		LSPMConsoleCommandHelper::ForEachLevelStreamingPersistenceManager([](ULevelStreamingPersistenceManager* Manager)
+		{
+			Manager->DumpContent();
+			return true;
+		});
 	})
 );
 
 FAutoConsoleCommand ULevelStreamingPersistenceManager::SetPropertyValueCommand(
-	TEXT("s.LevelStreamingPersistence.SetPropertyValue"),
+	TEXT("s.LevelStreamingPersistence.Debug.SetPropertyValue"),
 	TEXT("Set the persistent property's value for a given object"),
 	FConsoleCommandWithArgsDelegate::CreateLambda([](const TArray<FString>& InArgs)
 	{
@@ -828,50 +1248,44 @@ FAutoConsoleCommand ULevelStreamingPersistenceManager::SetPropertyValueCommand(
 			return;
 		}
 
-		for (const FWorldContext& Context : GEngine->GetWorldContexts())
+		LSPMConsoleCommandHelper::ForEachLevelStreamingPersistenceManager([&InArgs](ULevelStreamingPersistenceManager* Manager)
 		{
-			UWorld* World = Context.World();
-			if (World && World->IsGameWorld())
+			const FString& ObjectPath = InArgs[0];
+			const FName PropertyName = FName(InArgs[1]);
+			const FString& PropertyValue = InArgs[2];
+
+			if (LSPMConsoleCommandHelper::TrySetPropertyValueFromString<bool>(Manager, ObjectPath, PropertyName, PropertyValue) ||
+				LSPMConsoleCommandHelper::TrySetPropertyValueFromString<int32>(Manager, ObjectPath, PropertyName, PropertyValue) ||
+				LSPMConsoleCommandHelper::TrySetPropertyValueFromString<int64>(Manager, ObjectPath, PropertyName, PropertyValue) ||
+				LSPMConsoleCommandHelper::TrySetPropertyValueFromString<float>(Manager, ObjectPath, PropertyName, PropertyValue) ||
+				LSPMConsoleCommandHelper::TrySetPropertyValueFromString<double>(Manager, ObjectPath, PropertyName, PropertyValue) ||
+				LSPMConsoleCommandHelper::TrySetPropertyValueFromString<FName>(Manager, ObjectPath, PropertyName, PropertyValue) ||
+				LSPMConsoleCommandHelper::TrySetPropertyValueFromString<FString>(Manager, ObjectPath, PropertyName, PropertyValue) ||
+				LSPMConsoleCommandHelper::TrySetPropertyValueFromString<FText>(Manager, ObjectPath, PropertyName, PropertyValue) ||
+				LSPMConsoleCommandHelper::TrySetPropertyValueFromString<FVector>(Manager, ObjectPath, PropertyName, PropertyValue) ||
+				LSPMConsoleCommandHelper::TrySetPropertyValueFromString<FRotator>(Manager, ObjectPath, PropertyName, PropertyValue) ||
+				LSPMConsoleCommandHelper::TrySetPropertyValueFromString<FTransform>(Manager, ObjectPath, PropertyName, PropertyValue) ||
+				LSPMConsoleCommandHelper::TrySetPropertyValueFromString<FColor>(Manager, ObjectPath, PropertyName, PropertyValue) ||
+				LSPMConsoleCommandHelper::TrySetPropertyValueFromString<FLinearColor>(Manager, ObjectPath, PropertyName, PropertyValue))
 			{
-				if (ULevelStreamingPersistenceManager* Manager = World->GetSubsystem<ULevelStreamingPersistenceManager>())
+				return false;
+			}
+			else if (Manager->TrySetPropertyValueFromString(ObjectPath, PropertyName, PropertyValue))
+			{
+				FString Result;
+				if (Manager->GetPropertyValueAsString(ObjectPath, PropertyName, Result))
 				{
-					const FString& ObjectPath = InArgs[0];
-					const FName PropertyName = FName(InArgs[1]);
-					const FString& PropertyValue = InArgs[2];
-
-					if (LSPMConsoleCommandHelper::TrySetPropertyValueFromString<bool>(Manager, ObjectPath, PropertyName, PropertyValue) ||
-						LSPMConsoleCommandHelper::TrySetPropertyValueFromString<int32>(Manager, ObjectPath, PropertyName, PropertyValue) ||
-						LSPMConsoleCommandHelper::TrySetPropertyValueFromString<int64>(Manager, ObjectPath, PropertyName, PropertyValue) ||
-						LSPMConsoleCommandHelper::TrySetPropertyValueFromString<float>(Manager, ObjectPath, PropertyName, PropertyValue) ||
-						LSPMConsoleCommandHelper::TrySetPropertyValueFromString<double>(Manager, ObjectPath, PropertyName, PropertyValue) ||
-						LSPMConsoleCommandHelper::TrySetPropertyValueFromString<FName>(Manager, ObjectPath, PropertyName, PropertyValue) ||
-						LSPMConsoleCommandHelper::TrySetPropertyValueFromString<FString>(Manager, ObjectPath, PropertyName, PropertyValue) ||
-						LSPMConsoleCommandHelper::TrySetPropertyValueFromString<FText>(Manager, ObjectPath, PropertyName, PropertyValue) ||
-						LSPMConsoleCommandHelper::TrySetPropertyValueFromString<FVector>(Manager, ObjectPath, PropertyName, PropertyValue) ||
-						LSPMConsoleCommandHelper::TrySetPropertyValueFromString<FRotator>(Manager, ObjectPath, PropertyName, PropertyValue) ||
-						LSPMConsoleCommandHelper::TrySetPropertyValueFromString<FTransform>(Manager, ObjectPath, PropertyName, PropertyValue) ||
-						LSPMConsoleCommandHelper::TrySetPropertyValueFromString<FColor>(Manager, ObjectPath, PropertyName, PropertyValue) ||
-						LSPMConsoleCommandHelper::TrySetPropertyValueFromString<FLinearColor>(Manager, ObjectPath, PropertyName, PropertyValue))
-					{
-						return;
-					}
-					else if (Manager->TrySetPropertyValueFromString(ObjectPath, PropertyName, PropertyValue))
-					{
-						FString Result;
-						if (Manager->GetPropertyValueAsString(ObjectPath, PropertyName, Result))
-						{
-							UE_LOG(LogLevelStreamingPersistenceManager, Log, TEXT("GetPropertyValueAsString succeded : Property[%s] = %s for object %s"), *PropertyName.ToString(), *Result, *ObjectPath);
-							return;
-						}
-					}
+					UE_LOG(LogLevelStreamingPersistenceManager, Log, TEXT("GetPropertyValueAsString succeded : Property[%s] = %s for object %s"), *PropertyName.ToString(), *Result, *ObjectPath);
+					return false;
 				}
-			}
-		}
+			}			
+			return true;
+		});
 	})
 );
 
 FAutoConsoleCommand ULevelStreamingPersistenceManager::GetPropertyValueCommand(
-	TEXT("s.LevelStreamingPersistence.GetPropertyValue"),
+	TEXT("s.LevelStreamingPersistence.Debug.GetPropertyValue"),
 	TEXT("Get the persistent property's value for a given object"),
 	FConsoleCommandWithArgsDelegate::CreateLambda([](const TArray<FString>& InArgs)
 	{
@@ -880,46 +1294,95 @@ FAutoConsoleCommand ULevelStreamingPersistenceManager::GetPropertyValueCommand(
 			return;
 		}
 
-		for (const FWorldContext& Context : GEngine->GetWorldContexts())
+		LSPMConsoleCommandHelper::ForEachLevelStreamingPersistenceManager([&InArgs](ULevelStreamingPersistenceManager* Manager)
 		{
-			UWorld* World = Context.World();
-			if (World && World->IsGameWorld())
+			const FString& ObjectPath = InArgs[0];
+			const FName PropertyName = FName(InArgs[1]);
+
+			FString PropertyValue;
+			if (LSPMConsoleCommandHelper::GetPropertyValueAsString<bool>(Manager, ObjectPath, PropertyName, PropertyValue) ||
+				LSPMConsoleCommandHelper::GetPropertyValueAsString<int32>(Manager, ObjectPath, PropertyName, PropertyValue) ||
+				LSPMConsoleCommandHelper::GetPropertyValueAsString<int64>(Manager, ObjectPath, PropertyName, PropertyValue) ||
+				LSPMConsoleCommandHelper::GetPropertyValueAsString<float>(Manager, ObjectPath, PropertyName, PropertyValue) ||
+				LSPMConsoleCommandHelper::GetPropertyValueAsString<double>(Manager, ObjectPath, PropertyName, PropertyValue) ||
+				LSPMConsoleCommandHelper::GetPropertyValueAsString<FName>(Manager, ObjectPath, PropertyName, PropertyValue) ||
+				LSPMConsoleCommandHelper::GetPropertyValueAsString<FString>(Manager, ObjectPath, PropertyName, PropertyValue) ||
+				LSPMConsoleCommandHelper::GetPropertyValueAsString<FText>(Manager, ObjectPath, PropertyName, PropertyValue) ||
+				LSPMConsoleCommandHelper::GetPropertyValueAsString<FVector>(Manager, ObjectPath, PropertyName, PropertyValue) ||
+				LSPMConsoleCommandHelper::GetPropertyValueAsString<FRotator>(Manager, ObjectPath, PropertyName, PropertyValue) ||
+				LSPMConsoleCommandHelper::GetPropertyValueAsString<FTransform>(Manager, ObjectPath, PropertyName, PropertyValue) ||
+				LSPMConsoleCommandHelper::GetPropertyValueAsString<FColor>(Manager, ObjectPath, PropertyName, PropertyValue) ||
+				LSPMConsoleCommandHelper::GetPropertyValueAsString<FLinearColor>(Manager, ObjectPath, PropertyName, PropertyValue))
+			{
+				return false;
+			}
+			else
 			{
-				if (ULevelStreamingPersistenceManager* Manager = World->GetSubsystem<ULevelStreamingPersistenceManager>())
+				FString Result;
+				if (Manager->GetPropertyValueAsString(ObjectPath, PropertyName, Result))
 				{
-					const FString& ObjectPath = InArgs[0];
-					const FName PropertyName = FName(InArgs[1]);
-
-					FString PropertyValue;
-					if (LSPMConsoleCommandHelper::GetPropertyValueAsString<bool>(Manager, ObjectPath, PropertyName, PropertyValue) ||
-						LSPMConsoleCommandHelper::GetPropertyValueAsString<int32>(Manager, ObjectPath, PropertyName, PropertyValue) ||
-						LSPMConsoleCommandHelper::GetPropertyValueAsString<int64>(Manager, ObjectPath, PropertyName, PropertyValue) ||
-						LSPMConsoleCommandHelper::GetPropertyValueAsString<float>(Manager, ObjectPath, PropertyName, PropertyValue) ||
-						LSPMConsoleCommandHelper::GetPropertyValueAsString<double>(Manager, ObjectPath, PropertyName, PropertyValue) ||
-						LSPMConsoleCommandHelper::GetPropertyValueAsString<FName>(Manager, ObjectPath, PropertyName, PropertyValue) ||
-						LSPMConsoleCommandHelper::GetPropertyValueAsString<FString>(Manager, ObjectPath, PropertyName, PropertyValue) ||
-						LSPMConsoleCommandHelper::GetPropertyValueAsString<FText>(Manager, ObjectPath, PropertyName, PropertyValue) ||
-						LSPMConsoleCommandHelper::GetPropertyValueAsString<FVector>(Manager, ObjectPath, PropertyName, PropertyValue) ||
-						LSPMConsoleCommandHelper::GetPropertyValueAsString<FRotator>(Manager, ObjectPath, PropertyName, PropertyValue) ||
-						LSPMConsoleCommandHelper::GetPropertyValueAsString<FTransform>(Manager, ObjectPath, PropertyName, PropertyValue) ||
-						LSPMConsoleCommandHelper::GetPropertyValueAsString<FColor>(Manager, ObjectPath, PropertyName, PropertyValue) ||
-						LSPMConsoleCommandHelper::GetPropertyValueAsString<FLinearColor>(Manager, ObjectPath, PropertyName, PropertyValue))
-					{
-						return;
-					}
-					else
-					{
-						FString Result;
-						if (Manager->GetPropertyValueAsString(ObjectPath, PropertyName, Result))
-						{
-							UE_LOG(LogLevelStreamingPersistenceManager, Log, TEXT("GetPropertyValueAsString succeded : Property[%s] = %s for object %s"), *PropertyName.ToString(), *Result, *ObjectPath);
-							return;
-						}
-					}
+					UE_LOG(LogLevelStreamingPersistenceManager, Log, TEXT("GetPropertyValueAsString succeded : Property[%s] = %s for object %s"), *PropertyName.ToString(), *Result, *ObjectPath);
+					return false;
 				}
 			}
-		}
+			return true;
+		});
 	})
 );
 
-#endif // !UE_BUILD_SHIPPING
\ No newline at end of file
+FAutoConsoleCommand ULevelStreamingPersistenceManager::SaveToFileCommand(
+	TEXT("s.LevelStreamingPersistence.Debug.SaveToFile"),
+	TEXT("Save the content of LevelStreamingPersistenceManager to a file"),
+	FConsoleCommandWithArgsDelegate::CreateLambda([](const TArray<FString>& Args)
+	{
+		LSPMConsoleCommandHelper::ForEachLevelStreamingPersistenceManager([](ULevelStreamingPersistenceManager* Manager)
+		{
+			FString Filename = LSPMConsoleCommandHelper::GetFilename(Manager->GetWorld());
+			TUniquePtr<FArchive> FileWriter = TUniquePtr<FArchive>(IFileManager::Get().CreateFileWriter(*Filename));
+			if (FileWriter.IsValid())
+			{
+				TArray<uint8> Payload;
+				if (Manager->SerializeTo(Payload))
+				{
+					UE_LOG(LogLevelStreamingPersistenceManager, Log, TEXT("SaveToFile %s succeeded"), *Filename);
+					*FileWriter << Payload;
+					return false;
+				}
+			}
+			return true;
+		});
+	})
+);
+
+FAutoConsoleCommand ULevelStreamingPersistenceManager::LoadFromFileCommand(
+	TEXT("s.LevelStreamingPersistence.Debug.LoadFromFile"),
+	TEXT("Load from a file and initializes LevelStreamingPersistenceManager"),
+	FConsoleCommandWithArgsDelegate::CreateLambda([](const TArray<FString>& Args)
+	{
+		LSPMConsoleCommandHelper::ForEachLevelStreamingPersistenceManager([](ULevelStreamingPersistenceManager* Manager)
+		{
+			FString Filename = LSPMConsoleCommandHelper::GetFilename(Manager->GetWorld());
+			TUniquePtr<FArchive> FileReader = TUniquePtr<FArchive>(IFileManager::Get().CreateFileReader(*Filename));
+			if (FileReader.IsValid())
+			{
+				TArray<uint8> Payload;
+				*FileReader << Payload;
+				if (Manager->InitializeFrom(Payload))
+				{
+					UE_LOG(LogLevelStreamingPersistenceManager, Log, TEXT("LoadFromFile %s succeeded"), *Filename);
+					return false;
+				}
+			}
+			return true;
+		});
+	})
+);
+
+#endif // !UE_BUILD_SHIPPING
+
+bool ULevelStreamingPersistenceManager::bIsEnabled = true;
+FAutoConsoleVariableRef ULevelStreamingPersistenceManager::EnableCommand(
+	TEXT("s.LevelStreamingPersistence.Enabled"),
+	ULevelStreamingPersistenceManager::bIsEnabled,
+	TEXT("Turn on/off to enable/disable world persistent subsystem."),
+	ECVF_Default);
\ No newline at end of file
diff --git a/Engine/Plugins/Runtime/LevelStreamingPersistence/Source/LevelStreamingPersistence/Private/LevelStreamingPersistentObjectPropertyBag.cpp b/Engine/Plugins/Runtime/LevelStreamingPersistence/Source/LevelStreamingPersistence/Private/LevelStreamingPersistentObjectPropertyBag.cpp
--- a/Engine/Plugins/Runtime/LevelStreamingPersistence/Source/LevelStreamingPersistence/Private/LevelStreamingPersistentObjectPropertyBag.cpp
+++ b/Engine/Plugins/Runtime/LevelStreamingPersistence/Source/LevelStreamingPersistence/Private/LevelStreamingPersistentObjectPropertyBag.cpp
@@ -4,7 +4,12 @@
 #include "UObject/PropertyAccessUtil.h"
 #include "PropertyPathHelpers.h"
 
-bool FLevelStreamingPersistentObjectPropertyBag::Initialize(TFunctionRef<const UPropertyBag* ()> Func)
+FLevelStreamingPersistentObjectPropertyBag::FLevelStreamingPersistentObjectPropertyBag(FLevelStreamingPersistentObjectPropertyBag&& InOther)
+: PropertyBag(MoveTemp(InOther.PropertyBag))
+{
+}
+
+bool FLevelStreamingPersistentObjectPropertyBag::Initialize(TFunctionRef<const UPropertyBag* ()> InFunc)
 {
 	if (IsValid())
 	{
@@ -12,7 +17,7 @@ bool FLevelStreamingPersistentObjectPropertyBag::Initialize(TFunctionRef<const U
 	}
 	else
 	{
-		const UPropertyBag* DefaultClass = Func();
+		const UPropertyBag* DefaultClass = InFunc();
 		if (ensure(DefaultClass))
 		{
 			PropertyBag.InitializeFromBagStruct(DefaultClass);
@@ -43,6 +48,22 @@ bool FLevelStreamingPersistentObjectPropertyBag::ComparePropertyValueWithObject(
 	return true;
 }
 
+const FProperty* FLevelStreamingPersistentObjectPropertyBag::GetCompatibleProperty(const FProperty* InObjectProperty) const
+{
+	check(IsValid());
+	const FProperty* BagProperty = InObjectProperty ? FindPropertyByName(InObjectProperty->GetFName()) : nullptr;
+	if (!BagProperty)
+	{
+		return nullptr;
+	}
+
+	if (PropertyAccessUtil::ArePropertiesCompatible(BagProperty, InObjectProperty))
+	{
+		return BagProperty;
+	}
+	return nullptr;
+}
+
 const FProperty* FLevelStreamingPersistentObjectPropertyBag::CopyPropertyValueFromObject(const UObject* InObject, const FProperty* InObjectProperty)
 {
 	check(IsValid());
@@ -57,6 +78,20 @@ const FProperty* FLevelStreamingPersistentObjectPropertyBag::CopyPropertyValueFr
 	return (Result == EPropertyAccessResultFlags::Success) ? BagProperty : nullptr;
 }
 
+const FProperty* FLevelStreamingPersistentObjectPropertyBag::CopyPropertyValueFromPropertyBag(const FLevelStreamingPersistentObjectPropertyBag& InSourcePropertyBag, const FProperty* InSourceProperty)
+{
+	check(IsValid());
+	const FProperty* TargetBagProperty = InSourceProperty ? FindPropertyByName(InSourceProperty->GetFName()) : nullptr;
+	if (!TargetBagProperty)
+	{
+		return nullptr;
+	}
+
+	const void* TargetBagValue = TargetBagProperty->ContainerPtrToValuePtr<void>(PropertyBag.GetValue().GetMemory());
+	EPropertyAccessResultFlags Result = PropertyAccessUtil::GetPropertyValue_InContainer(InSourceProperty, InSourcePropertyBag.PropertyBag.GetValue().GetMemory(), TargetBagProperty, const_cast<void*>(TargetBagValue), INDEX_NONE);
+	return (Result == EPropertyAccessResultFlags::Success) ? TargetBagProperty : nullptr;
+}
+
 const FProperty* FLevelStreamingPersistentObjectPropertyBag::CopyPropertyValueToObject(UObject* InObject, const FProperty* InObjectProperty) const
 {
 	check(IsValid());
@@ -127,6 +162,17 @@ void FLevelStreamingPersistentObjectPropertyBag::DumpContent(TFunctionRef<void(c
 	}
 }
 
+void FLevelStreamingPersistentObjectPropertyBag::Serialize(FArchive& Ar)
+{
+	PropertyBag.Serialize(Ar);
+}
+
+FArchive& operator<<(FArchive& Ar, FLevelStreamingPersistentObjectPropertyBag& PropertyBag)
+{
+	PropertyBag.Serialize(Ar);
+	return Ar;
+}
+
 template<> struct FLevelStreamingPersistentObjectPropertyBag::TPropertyBagPropertyType<bool> { static const EPropertyBagPropertyType Value = EPropertyBagPropertyType::Bool; };
 template<> struct FLevelStreamingPersistentObjectPropertyBag::TPropertyBagPropertyType<int32> { static const EPropertyBagPropertyType Value = EPropertyBagPropertyType::Int32; };
 template<> struct FLevelStreamingPersistentObjectPropertyBag::TPropertyBagPropertyType<int64> { static const EPropertyBagPropertyType Value = EPropertyBagPropertyType::Int64; };
diff --git a/Engine/Plugins/Runtime/LevelStreamingPersistence/Source/LevelStreamingPersistence/Public/LevelStreamingPersistenceManager.h b/Engine/Plugins/Runtime/LevelStreamingPersistence/Source/LevelStreamingPersistence/Public/LevelStreamingPersistenceManager.h
--- a/Engine/Plugins/Runtime/LevelStreamingPersistence/Source/LevelStreamingPersistence/Public/LevelStreamingPersistenceManager.h
+++ b/Engine/Plugins/Runtime/LevelStreamingPersistence/Source/LevelStreamingPersistence/Public/LevelStreamingPersistenceManager.h
@@ -19,38 +19,67 @@ using FLevelStreamingPersistentPropertyArray = TArray<FCustomPropertyListNode, T
 
 struct FLevelStreamingPersistentObjectPrivateProperties
 {
-	// Payload data containing serialized changed properties
-	TArray<uint8> PayloadData;
+	FLevelStreamingPersistentObjectPrivateProperties() = default;
+	FLevelStreamingPersistentObjectPrivateProperties(FLevelStreamingPersistentObjectPrivateProperties&& Other);
+
 	// Initial snapshot used detect changed properties
 	FLevelStreamingPersistentObjectPropertyBag Snapshot;
+
+	// Object source class
+	FString SourceClassPathName;
+	// Payload data containing serialized changed properties
+	TArray<uint8> PayloadData;
 	// Persistent Properties
 	TArray<const FProperty*> PersistentProperties;
-	
+
+	const UClass* GetSourceClass() const { return LoadObject<UClass>(nullptr, *SourceClassPathName); }
 	void ForEachPersistentProperty(TFunctionRef<void(const FProperty*)> Func) const;
-	void BuildSerializedPropertyList(FLevelStreamingPersistentPropertyArray& OutCustomPropertyList) const;
+	bool Sanitize(const ULevelStreamingPersistenceManager& InManager);
+	void Serialize(FArchive& Ar);
+	friend FArchive& operator<<(FArchive& Ar, FLevelStreamingPersistentObjectPrivateProperties& PrivateProperties);
 };
 
 struct FLevelStreamingPersistentObjectPublicProperties
 {
+	FLevelStreamingPersistentObjectPublicProperties() = default;
+	FLevelStreamingPersistentObjectPublicProperties(FLevelStreamingPersistentObjectPublicProperties&& Other);
+
+	// Properties to persist (computed once)
+	TArray<const FProperty*> PropertiesToPersist;
+
+	// Object source class
+	FString SourceClassPathName;
 	// Public properties
 	FLevelStreamingPersistentObjectPropertyBag PropertyBag;
 	// Persistent Properties
 	TArray<const FProperty*> PersistentProperties;
 
+	const UClass* GetSourceClass() const { return LoadObject<UClass>(nullptr, *SourceClassPathName); }
 	void ForEachPersistentProperty(TFunctionRef<void(const FProperty*)> Func) const;
+	bool Sanitize(const ULevelStreamingPersistenceManager& InManager);
+	void Serialize(FArchive& Ar);
+	friend FArchive& operator<<(FArchive& Ar, FLevelStreamingPersistentObjectPublicProperties& PublicProperties);
 };
 
 struct FLevelStreamingPersistentPropertyValues
 {
+	bool bIsMakingVisibleCacheValid = false;
 	TMap<FString, FLevelStreamingPersistentObjectPrivateProperties> ObjectsPrivatePropertyValues;
 	TMap<FString, FLevelStreamingPersistentObjectPublicProperties> ObjectsPublicPropertyValues;
+
+	void Serialize(FArchive& Ar);
+	friend FArchive& operator<<(FArchive& Ar, FLevelStreamingPersistentPropertyValues& Properties);
 };
 
 // Helper class to access FLevelStreamingPersistentProperty's Properties
-class FLevelStreamingPersistentPropertiesInfo
+UCLASS()
+class ULevelStreamingPersistentPropertiesInfo : public UObject
 {
+	GENERATED_BODY()
+
 public:
 	void Initialize();
+	static void AddReferencedObjects(UObject* InThis, FReferenceCollector& Collector);
 
 	enum EPropertyType
 	{
@@ -61,16 +90,14 @@ public:
 
 	// Configurable persistent properties helper methods
 	const UPropertyBag* GetPropertyBagFromClass(EPropertyType InAccessSpecifier, const UClass* InClass) const;
-	const UClass* GetClassFromPropertyBag(EPropertyType InAccessSpecifier, const UPropertyBag* InPropertyBag) const;
 	void ForEachProperty(EPropertyType InAccessSpecifier, const UClass* InClass, TFunctionRef<void(FProperty*)> Func) const;
 	bool HasProperties(EPropertyType InAccessSpecifier, const UClass* InClass) const;
 	bool HasProperty(EPropertyType InAccessSpecifier, const FProperty* InProperty) const;
 
 private:
 	/* Acceleration maps to find properties for a given class */
-	TStaticArray<TMap<TWeakObjectPtr<const UClass>, TSet<FProperty*>>, PropertyType_Count> ClassesProperties;
-	TStaticArray<TMap<TWeakObjectPtr<const UClass>, FInstancedPropertyBag>, PropertyType_Count> ObjectClassToPropertyBag;
-	TStaticArray<TMap<const UPropertyBag*, TWeakObjectPtr<const UClass>>, PropertyType_Count> PropertyBagToObjectClass;
+	TStaticArray<TMap<const UClass*, TSet<FProperty*>>, PropertyType_Count> ClassesProperties;
+	TStaticArray<TMap<const UClass*, FInstancedPropertyBag>, PropertyType_Count> ObjectClassToPropertyBag;
 };
 
 UCLASS()
@@ -84,10 +111,13 @@ public:
 	~ULevelStreamingPersistenceManager() {}
 	virtual bool ShouldCreateSubsystem(UObject* Outer) const override;
 	virtual void Initialize(FSubsystemCollectionBase& Collection) override;
-	virtual void PostInitialize() override;
 	virtual void Deinitialize() override;
 	//~ End WorldSubsystem
 
+	// Serialization
+	bool SerializeTo(TArray<uint8>& OutPayload);
+	bool InitializeFrom(const TArray<uint8>& InPayload);
+
 	// Sets property value and creates the entry if necessary, returns true on success.
 	template <typename ClassType, typename PropertyType>
 	bool SetPropertyValue(const FString& InObjectPathName, const FName InPropertyName, const PropertyType& InPropertyValue);
@@ -118,7 +148,7 @@ private:
 	bool SaveLevelPersistentPropertyValues(const ULevel* InLevel);
 
 	// Restore persistent properties
-	bool RestoreLevelPersistentPropertyValues(ULevel* InLevel) const;
+	bool RestoreLevelPersistentPropertyValues(const ULevel* InLevel) const;
 	int32 RestorePrivateProperties(UObject* InObject, const FLevelStreamingPersistentObjectPrivateProperties& InPersistentProperties) const;
 	int32 RestorePublicProperties(UObject* InObject, const FLevelStreamingPersistentObjectPublicProperties& InPersistentProperties) const;
 
@@ -131,9 +161,12 @@ private:
 	const FString GetResolvedObjectPathName(const FString& InObjectPathName) const;
 	bool CopyPropertyBagValueToObject(const FLevelStreamingPersistentObjectPropertyBag* InPropertyBag, UObject* InObject, FProperty* InObjectProperty) const;
 	TPair<UObject*, FProperty*> GetObjectPropertyPair(const FString& InObjectPathName, const FName InPropertyName) const;
+#if !UE_BUILD_SHIPPING
 	void DumpContent() const;
+#endif
 	const FLevelStreamingPersistentObjectPropertyBag* GetPropertyBag(const FString& InObjectPathName) const;
 	FLevelStreamingPersistentObjectPropertyBag* GetPropertyBag(const FString& InObjectPathName);
+	bool SerializeManager(FArchive& Ar);
 
 	template <typename ClassType>
 	FLevelStreamingPersistentObjectPropertyBag* GetOrCreatePropertyBag(const FString& InObjectPathName, const FName InPropertyName);
@@ -142,6 +175,8 @@ private:
 	static class FAutoConsoleVariableRef EnableCommand;
 	static bool bIsEnabled;
 #if !UE_BUILD_SHIPPING
+	static class FAutoConsoleCommand SaveToFileCommand;
+	static class FAutoConsoleCommand LoadFromFileCommand;
 	static class FAutoConsoleCommand DumpContentCommand;
 	static class FAutoConsoleCommand SetPropertyValueCommand;
 	static class FAutoConsoleCommand GetPropertyValueCommand;
@@ -154,8 +189,14 @@ private:
 	ILevelStreamingPersistenceModule* PersistenceModule = nullptr;
 
 	// Persistent Properties Info
-	FLevelStreamingPersistentPropertiesInfo PersistentPropertiesInfo;
+	UPROPERTY(Transient)
+	TObjectPtr<ULevelStreamingPersistentPropertiesInfo> PersistentPropertiesInfo;
+
+	// Whether to use tagget property serialization
+	bool bUseTaggedPropertySerialization;
 
+	friend struct FLevelStreamingPersistentObjectPrivateProperties;
+	friend struct FLevelStreamingPersistentObjectPublicProperties;
 	friend class FPersistenceModule;
 };
 
@@ -212,12 +253,12 @@ FLevelStreamingPersistentObjectPropertyBag* ULevelStreamingPersistenceManager::G
 
 	if (FProperty* ObjectProperty = ObjectClass->FindPropertyByName(InPropertyName))
 	{
-		if (PersistentPropertiesInfo.HasProperty(FLevelStreamingPersistentPropertiesInfo::PropertyType_Public, ObjectProperty))
+		if (PersistentPropertiesInfo->HasProperty(ULevelStreamingPersistentPropertiesInfo::PropertyType_Public, ObjectProperty))
 		{
 			FLevelStreamingPersistentPropertyValues& LevelProperties = LevelsPropertyValues.FindOrAdd(LevelPathName);
 
 			FLevelStreamingPersistentObjectPublicProperties& ObjectPublicPropertyValues = LevelProperties.ObjectsPublicPropertyValues.FindOrAdd(ObjectPathName);
-			if (ObjectPublicPropertyValues.PropertyBag.Initialize([this, ObjectClass]() { return PersistentPropertiesInfo.GetPropertyBagFromClass(FLevelStreamingPersistentPropertiesInfo::PropertyType_Public, ObjectClass); }))
+			if (ObjectPublicPropertyValues.PropertyBag.Initialize([this, ObjectClass]() { return PersistentPropertiesInfo->GetPropertyBagFromClass(ULevelStreamingPersistentPropertiesInfo::PropertyType_Public, ObjectClass); }))
 			{
 				return &ObjectPublicPropertyValues.PropertyBag;
 			}
diff --git a/Engine/Plugins/Runtime/LevelStreamingPersistence/Source/LevelStreamingPersistence/Public/LevelStreamingPersistenceSettings.h b/Engine/Plugins/Runtime/LevelStreamingPersistence/Source/LevelStreamingPersistence/Public/LevelStreamingPersistenceSettings.h
--- a/Engine/Plugins/Runtime/LevelStreamingPersistence/Source/LevelStreamingPersistence/Public/LevelStreamingPersistenceSettings.h
+++ b/Engine/Plugins/Runtime/LevelStreamingPersistence/Source/LevelStreamingPersistence/Public/LevelStreamingPersistenceSettings.h
@@ -25,5 +25,5 @@ class LEVELSTREAMINGPERSISTENCE_API ULevelStreamingPersistenceSettings : public
 	UPROPERTY(Config)
 	TArray<FLevelStreamingPersistentProperty> Properties;
 
-	friend class FLevelStreamingPersistentPropertiesInfo;
+	friend class ULevelStreamingPersistentPropertiesInfo;
 };
\ No newline at end of file
diff --git a/Engine/Plugins/Runtime/LevelStreamingPersistence/Source/LevelStreamingPersistence/Public/LevelStreamingPersistentObjectPropertyBag.h b/Engine/Plugins/Runtime/LevelStreamingPersistence/Source/LevelStreamingPersistence/Public/LevelStreamingPersistentObjectPropertyBag.h
--- a/Engine/Plugins/Runtime/LevelStreamingPersistence/Source/LevelStreamingPersistence/Public/LevelStreamingPersistentObjectPropertyBag.h
+++ b/Engine/Plugins/Runtime/LevelStreamingPersistence/Source/LevelStreamingPersistence/Public/LevelStreamingPersistentObjectPropertyBag.h
@@ -8,15 +8,21 @@
 
 struct FLevelStreamingPersistentObjectPropertyBag
 {
-	bool Initialize(TFunctionRef<const UPropertyBag* ()> Func);
+	FLevelStreamingPersistentObjectPropertyBag() = default;
+	FLevelStreamingPersistentObjectPropertyBag(FLevelStreamingPersistentObjectPropertyBag&& InOther);
+	bool Initialize(TFunctionRef<const UPropertyBag* ()> InFunc);
 	bool IsValid() const { return PropertyBag.IsValid(); }
 	const UPropertyBag* GetPropertyBagStruct() const { return PropertyBag.GetPropertyBagStruct(); }
 	const FProperty* FindPropertyByName(const FName InPropertyName) const;
 	const FProperty* CopyPropertyValueFromObject(const UObject* InObject, const FProperty* InObjectProperty);
+	const FProperty* CopyPropertyValueFromPropertyBag(const FLevelStreamingPersistentObjectPropertyBag& InPropertyBag, const FProperty* InProperty);
 	const FProperty* CopyPropertyValueToObject(UObject* InObject, const FProperty* InObjectProperty) const;
+	const FProperty* GetCompatibleProperty(const FProperty* InObjectProperty) const;
 	bool ComparePropertyValueWithObject(const UObject* InObject, const FProperty* InObjectProperty, bool& bOutIsIdentical) const;
 	void ForEachProperty(TFunctionRef<void(const FProperty*)> Func) const;
 	void DumpContent(TFunctionRef<void(const FProperty*, const FString&)> Func, const TArray<const FProperty*>* InDumpProperties = nullptr);
+	void Serialize(FArchive& Ar);
+	friend FArchive& operator<<(FArchive& Ar, FLevelStreamingPersistentObjectPropertyBag& PropertyBag);
 
 	template<typename PropertyType>
 	bool SetPropertyValue(const FName InPropertyName, const PropertyType& InPropertyValue);
