Fixed PVS Static Analysis warning V1053

diff --git a/Engine/Plugins/Runtime/LevelStreamingPersistence/Source/LevelStreamingPersistence/Private/LevelStreamingPersistenceManager.cpp b/Engine/Plugins/Runtime/LevelStreamingPersistence/Source/LevelStreamingPersistence/Private/LevelStreamingPersistenceManager.cpp
--- a/Engine/Plugins/Runtime/LevelStreamingPersistence/Source/LevelStreamingPersistence/Private/LevelStreamingPersistenceManager.cpp
+++ b/Engine/Plugins/Runtime/LevelStreamingPersistence/Source/LevelStreamingPersistence/Private/LevelStreamingPersistenceManager.cpp
@@ -49,7 +49,7 @@ public:
 		SetIsTextFormat(InArchive.IsTextFormat());
 		SetWantBinaryPropertySerialization(InArchive.WantBinaryPropertySerialization());
 		SetIsPersistent(true);
-		SetFilterEditorOnly(true);
+		FArchiveProxy::SetFilterEditorOnly(true);
 		ArShouldSkipBulkData = true;
 
 		if (WantBinaryPropertySerialization())
