#UE IWYU fix

diff --git a/Engine/Plugins/Runtime/LevelStreamingPersistence/Source/LevelStreamingPersistence/Private/LevelStreamingPersistenceManager.cpp b/Engine/Plugins/Runtime/LevelStreamingPersistence/Source/LevelStreamingPersistence/Private/LevelStreamingPersistenceManager.cpp
--- a/Engine/Plugins/Runtime/LevelStreamingPersistence/Source/LevelStreamingPersistence/Private/LevelStreamingPersistenceManager.cpp
+++ b/Engine/Plugins/Runtime/LevelStreamingPersistence/Source/LevelStreamingPersistence/Private/LevelStreamingPersistenceManager.cpp
@@ -4,6 +4,7 @@
 #include "LevelStreamingPersistenceModule.h"
 #include "LevelStreamingPersistenceSettings.h"
 #include "Serialization/ArchiveUObject.h"
+#include "Serialization/CustomVersion.h"
 #include "Serialization/ObjectAndNameAsStringProxyArchive.h"
 #include "Serialization/MemoryReader.h"
 #include "Serialization/MemoryWriter.h"
@@ -16,6 +17,7 @@
 #include "Engine/Level.h"
 #include "Engine/Engine.h"
 #include "GameFramework/WorldSettings.h"
+#include "UObject/Package.h"
 
 #if !UE_BUILD_SHIPPING
 #include "HAL/FileManager.h"
