[LevelStreamingPersistence]
* Fixed linker error when building non-unity

diff --git a/Engine/Plugins/Runtime/LevelStreamingPersistence/Source/LevelStreamingPersistence/Private/LevelStreamingPersistenceManager.cpp b/Engine/Plugins/Runtime/LevelStreamingPersistence/Source/LevelStreamingPersistence/Private/LevelStreamingPersistenceManager.cpp
--- a/Engine/Plugins/Runtime/LevelStreamingPersistence/Source/LevelStreamingPersistence/Private/LevelStreamingPersistenceManager.cpp
+++ b/Engine/Plugins/Runtime/LevelStreamingPersistence/Source/LevelStreamingPersistence/Private/LevelStreamingPersistenceManager.cpp
@@ -24,6 +24,8 @@
 #include "Misc/Paths.h"
 #endif
 
+#include UE_INLINE_GENERATED_CPP_BY_NAME(LevelStreamingPersistenceManager)
+
 DEFINE_LOG_CATEGORY_STATIC(LogLevelStreamingPersistenceManager, All, All);
 
 /*
diff --git a/Engine/Plugins/Runtime/LevelStreamingPersistence/Source/LevelStreamingPersistence/Public/LevelStreamingPersistentObjectPropertyBag.h b/Engine/Plugins/Runtime/LevelStreamingPersistence/Source/LevelStreamingPersistence/Public/LevelStreamingPersistentObjectPropertyBag.h
--- a/Engine/Plugins/Runtime/LevelStreamingPersistence/Source/LevelStreamingPersistence/Public/LevelStreamingPersistentObjectPropertyBag.h
+++ b/Engine/Plugins/Runtime/LevelStreamingPersistence/Source/LevelStreamingPersistence/Public/LevelStreamingPersistentObjectPropertyBag.h
@@ -35,10 +35,10 @@ struct FLevelStreamingPersistentObjectPropertyBag
 
 private:
 	template<typename PropertyType>
-	EPropertyBagResult SetValue(const FName InPropertyName, const PropertyType& InPropertyValue) { return EPropertyBagResult::PropertyNotFound; }
+	EPropertyBagResult SetValue(const FName InPropertyName, const PropertyType& InPropertyValue);
 	
 	template<typename PropertyType>
-	TValueOrError<PropertyType, EPropertyBagResult> GetValue(const FName InPropertyName) const { return MakeError(EPropertyBagResult::PropertyNotFound); }
+	TValueOrError<PropertyType, EPropertyBagResult> GetValue(const FName InPropertyName) const;
 
 	template<typename PropertyType>
 	TValueOrError<PropertyType, EPropertyBagResult> GetValueStruct(const FName InPropertyName) const;
