LevelStreamingPersistence: addressing crash when a map placed actor with persisted properties is destroyed during gameplay, and its old streaming level gets made visible again from memory
#Synced-CL 51821995

diff --git a/Engine/Plugins/Runtime/LevelStreamingPersistence/Source/LevelStreamingPersistence/Private/LevelStreamingPersistenceManager.cpp b/Engine/Plugins/Runtime/LevelStreamingPersistence/Source/LevelStreamingPersistence/Private/LevelStreamingPersistenceManager.cpp
--- a/Engine/Plugins/Runtime/LevelStreamingPersistence/Source/LevelStreamingPersistence/Private/LevelStreamingPersistenceManager.cpp
+++ b/Engine/Plugins/Runtime/LevelStreamingPersistence/Source/LevelStreamingPersistence/Private/LevelStreamingPersistenceManager.cpp
@@ -432,14 +432,20 @@ bool ULevelStreamingPersistenceManager::RestoreLevelPersistentPropertyValues(con
 	{
 		for (const auto& [ObjectPath, ObjectPrivatePropertyValues] : LevelProperties->ObjectsPrivatePropertyValues)
 		{
-			if (UObject* Object = FindObject<UObject>(nullptr, *ObjectPath))
+			// Find the object by path and only restore properties if it's valid.
+			// Checking validity because an object may exist in memory, yet may have been destroyed during gameplay
+			// before the level was streamed out. We don't want to restore properties on those stale objects,
+			// they're just pending GC.
+			UObject* Object = FindObject<UObject>(nullptr, *ObjectPath);
+			if (IsValid(Object))
 			{
 				RestoredCount += RestorePrivateProperties(Object, ObjectPrivatePropertyValues);
 			}
 		}
 		for (const auto& [ObjectPath, ObjectPublicPropertyValues] : LevelProperties->ObjectsPublicPropertyValues)
 		{
-			if (UObject* Object = FindObject<UObject>(nullptr, *ObjectPath))
+			UObject* Object = FindObject<UObject>(nullptr, *ObjectPath);
+			if (IsValid(Object))
 			{
 				RestoredCount += RestorePublicProperties(Object, ObjectPublicPropertyValues);
 			}
