LevelStreamingPersistence: Warn users when they've configured a BP variable to be persisted, when that variable isn't marked instance editable.
The value would get reset during RerunConstructionScript otherwise, so it would fail to persist in PIE. (It would get loaded and stomped)
#Synced-CL 51821995

diff --git a/Engine/Plugins/Runtime/LevelStreamingPersistence/Source/LevelStreamingPersistence/Private/LevelStreamingPersistenceManager.cpp b/Engine/Plugins/Runtime/LevelStreamingPersistence/Source/LevelStreamingPersistence/Private/LevelStreamingPersistenceManager.cpp
--- a/Engine/Plugins/Runtime/LevelStreamingPersistence/Source/LevelStreamingPersistence/Private/LevelStreamingPersistenceManager.cpp
+++ b/Engine/Plugins/Runtime/LevelStreamingPersistence/Source/LevelStreamingPersistence/Private/LevelStreamingPersistenceManager.cpp
@@ -13,9 +13,12 @@
 #include "WorldPartition/WorldPartitionHelpers.h"
 #include "StructUtils/PropertyBag.h"
 #include "PropertyPathHelpers.h"
+#include "Logging/MessageLog.h"
+#include "Misc/UObjectToken.h"
 #include "Engine/World.h"
 #include "Engine/Level.h"
 #include "Engine/Engine.h"
+#include "Engine/BlueprintGeneratedClass.h"
 #include "GameFramework/WorldSettings.h"
 #include "UObject/Package.h"
 
@@ -1023,6 +1026,18 @@ void ULevelStreamingPersistentPropertiesInfo::Initialize()
 			if (const UClass* Class = Property->GetOwnerClass())
 			{
 				OutClassesProperties.FindOrAdd(Class).Add(Property);
+
+#if WITH_EDITOR
+				// Blueprint variables must be marked as instance edit able to support persistence, otherwise in PIE the persisted value 
+				// will be reset to default during RerunConstructionScripts. Only affects PIE which reruns construction scripts at runtime.
+				if (Property->HasAnyPropertyFlags(EPropertyFlags::CPF_DisableEditOnInstance) && Cast<UBlueprintGeneratedClass>(Class) != nullptr)
+				{
+					FString WarningMessage = FString::Printf(TEXT("Make variable '%s' in '%s' Instance Editable to support persistence. "), *Property->GetName(), *Class->GetName());
+					FMessageLog MessageLog("LoadErrors");
+					MessageLog.Warning(FText::FromString(WarningMessage))->AddToken(FUObjectToken::Create(Class->GetPackage()));
+					MessageLog.Open(EMessageSeverity::Warning);
+				}
+#endif
 			}
 		}
 		else
