LevelStreamingPersistence: SaveGame utilities added
- BP exposed SerializeTo and InitializeFrom functions.
- Added a helper function to update save cache for all visible levels. People will want to call this before SerializeTo.
#Synced-CL 51821995

diff --git a/Engine/Plugins/Runtime/LevelStreamingPersistence/Source/LevelStreamingPersistence/Private/LevelStreamingPersistenceManager.cpp b/Engine/Plugins/Runtime/LevelStreamingPersistence/Source/LevelStreamingPersistence/Private/LevelStreamingPersistenceManager.cpp
--- a/Engine/Plugins/Runtime/LevelStreamingPersistence/Source/LevelStreamingPersistence/Private/LevelStreamingPersistenceManager.cpp
+++ b/Engine/Plugins/Runtime/LevelStreamingPersistence/Source/LevelStreamingPersistence/Private/LevelStreamingPersistenceManager.cpp
@@ -281,6 +281,34 @@ bool ULevelStreamingPersistenceManager::SaveLevelPersistentPropertyValues(const
 	return SavedCount > 0;
 }
 
+void ULevelStreamingPersistenceManager::UpdateVisibleLevelsPersistentPropertyValues()
+{
+	if (IsEnabled())
+	{
+		// Iterate all of the world's levels
+		UWorld* World = GetWorld();
+		for (ULevel* Level : World->GetLevels())
+		{
+			// Only refresh visible levels. If the level is not visible, it's state should be up-to-date already. 
+			// It would have been updated when the level was made invisible.
+			if (Level->bIsVisible)
+			{
+				// Only refresh levels if they already have an entry in LevelsPropertyValues. If the level is
+				// eligible for persistence, an entry for it should have been made when it became visible.
+				const FString LevelPathName = Level->GetPathName();
+				if (LevelsPropertyValues.Contains(LevelPathName))
+				{
+					SaveLevelPersistentPropertyValues(Level);
+				}
+			}
+		}
+	}
+	else
+	{
+		UE_LOG(LogLevelStreamingPersistenceManager, Warning, TEXT("Call to UpdateVisibleLevelsPersistentPropertyValues ignored because persistence is disabled by cvar"));
+	}
+}
+
 void ULevelStreamingPersistenceManager::OnLevelBeginMakingVisible(UWorld* World, const ULevelStreaming* InStreamingLevel, ULevel* InLoadedLevel)
 {
 	if (IsValid(InLoadedLevel) && (World == GetWorld()) && IsEnabled())
@@ -613,8 +641,14 @@ bool ULevelStreamingPersistenceManager::DiffWithSnapshot(const UObject* InObject
 	return false;
 }
 
-bool ULevelStreamingPersistenceManager::SerializeTo(TArray<uint8>& OutPayload)
+bool ULevelStreamingPersistenceManager::SerializeTo(TArray<uint8>& OutPayload, const bool bForceUpdate)
 {
+	// If requested, iterate all visible levels to ensure their latest state is stored, before serializing that.
+	if (bForceUpdate)
+	{
+		UpdateVisibleLevelsPersistentPropertyValues();
+	}
+
 	// Serialize to archive and gather custom versions
 	TArray<uint8> PayloadData;
 	FMemoryWriter PayloadAr(PayloadData, true);
diff --git a/Engine/Plugins/Runtime/LevelStreamingPersistence/Source/LevelStreamingPersistence/Public/LevelStreamingPersistenceManager.h b/Engine/Plugins/Runtime/LevelStreamingPersistence/Source/LevelStreamingPersistence/Public/LevelStreamingPersistenceManager.h
--- a/Engine/Plugins/Runtime/LevelStreamingPersistence/Source/LevelStreamingPersistence/Public/LevelStreamingPersistenceManager.h
+++ b/Engine/Plugins/Runtime/LevelStreamingPersistence/Source/LevelStreamingPersistence/Public/LevelStreamingPersistenceManager.h
@@ -114,8 +114,16 @@ public:
 	virtual void Deinitialize() override;
 	//~ End WorldSubsystem
 
-	// Serialization
-	bool SerializeTo(TArray<uint8>& OutPayload);
+	// Save game utility functions
+
+	// Serialize the persistent properties of all streaming levels of the current map to byte array. This can be put in a SaveGame's property to save to file.
+	// @param bForceUpdate If true, will iterate all visible streaming levels to update their persistence data. Recommended for saving latest world snapshot.
+	UFUNCTION(BlueprintCallable, Category = "Level Streaming Persistence")
+	bool SerializeTo(TArray<uint8>& OutPayload, const bool bForceUpdate = false);
+
+	// Deserialize the persistent properties of all streaming levels of the current map from byte array. The payload should be from an earlier SerializeTo call, 
+	// from the same world partition map, but it can be from an earlier play session.
+	UFUNCTION(BlueprintCallable, Category = "Level Streaming Persistence")
 	bool InitializeFrom(const TArray<uint8>& InPayload);
 
 	// Sets property value and creates the entry if necessary, returns true on success.
@@ -144,6 +152,11 @@ private:
 	void OnLevelBeginMakingInvisible(UWorld* World, const ULevelStreaming* StreamingLevel, ULevel* LoadedLevel);
 	void OnLevelBeginMakingVisible(UWorld* World, const ULevelStreaming* StreamingLevel, ULevel* LoadedLevel);
 
+	// Iterate all currently visible streaming levels to record the persistent properties of objects. Normally, these values are only updated 
+	// for a level when it is made invisible. Call this to make sure that the latest gameplay state of all streaming levels, visible and invisible, 
+	// is stored in the manager. You can then serialize this data, for example for save game purposes, by calling SerializeTo.
+	void UpdateVisibleLevelsPersistentPropertyValues();
+
 	// Save persistent properties
 	bool SaveLevelPersistentPropertyValues(const ULevel* InLevel);
 
