LevelStreamingPersistence: Showing FMessageLog warning and hint to users when a property path was not resolvable. This is to help people realize and fix their syntax errors.
#Synced-CL 52723097

diff --git a/Engine/Plugins/Runtime/LevelStreamingPersistence/Source/LevelStreamingPersistence/Private/LevelStreamingPersistenceManager.cpp b/Engine/Plugins/Runtime/LevelStreamingPersistence/Source/LevelStreamingPersistence/Private/LevelStreamingPersistenceManager.cpp
--- a/Engine/Plugins/Runtime/LevelStreamingPersistence/Source/LevelStreamingPersistence/Private/LevelStreamingPersistenceManager.cpp
+++ b/Engine/Plugins/Runtime/LevelStreamingPersistence/Source/LevelStreamingPersistence/Private/LevelStreamingPersistenceManager.cpp
@@ -31,6 +31,9 @@
 
 DEFINE_LOG_CATEGORY_STATIC(LogLevelStreamingPersistenceManager, All, All);
 
+
+#define LOCTEXT_NAMESPACE "LevelStreamingPersistence"
+
 /*
  * Achives used by ULevelStreamingPersistenceManager
  */
@@ -1051,7 +1054,7 @@ void ULevelStreamingPersistentPropertiesInfo::AddReferencedObjects(UObject* InTh
 
 void ULevelStreamingPersistentPropertiesInfo::Initialize()
 {
-	auto FillClassesProperties = [this](const FString& InPropertyPath, bool bIsPublic)
+	auto FillClassesProperties = [this](const FString& InPropertyPath, bool bIsPublic) -> bool
 	{
 		auto& OutClassesProperties = ClassesProperties[bIsPublic ? PropertyType_Public : PropertyType_Private];
 		FSoftObjectPath(InPropertyPath).TryLoad();
@@ -1066,18 +1069,18 @@ void ULevelStreamingPersistentPropertiesInfo::Initialize()
 				// will be reset to default during RerunConstructionScripts. Only affects PIE which reruns construction scripts at runtime.
 				if (Property->HasAnyPropertyFlags(EPropertyFlags::CPF_DisableEditOnInstance) && Cast<UBlueprintGeneratedClass>(Class) != nullptr)
 				{
-					FString WarningMessage = FString::Printf(TEXT("Make variable '%s' in '%s' Instance Editable to support persistence. "), *Property->GetName(), *Class->GetName());
+					FText WarningMessage = FText::Format(LOCTEXT("InstanceEditableFormat", "Make variable {0} in {1} Instance Editable to support persistence. "), FText::FromName(Property->GetFName()), FText::FromName(Class->GetFName()));
 					FMessageLog MessageLog("LoadErrors");
-					MessageLog.Warning(FText::FromString(WarningMessage))->AddToken(FUObjectToken::Create(Class->GetPackage()));
+					MessageLog.Warning(WarningMessage)->AddToken(FUObjectToken::Create(Class->GetPackage()));
 					MessageLog.Open(EMessageSeverity::Warning);
 				}
 #endif
+				return true;
 			}
 		}
-		else
-		{
-			UE_LOGF(LogLevelStreamingPersistenceManager, Log, "Could not resolve property path %ls.", *InPropertyPath);
-		}
+		
+		UE_LOGF(LogLevelStreamingPersistenceManager, Log, "Could not resolve property path %ls.", *InPropertyPath);
+		return false;
 	};
 
 	auto CreateClassDefaultPropertyBag = [this]()
@@ -1108,11 +1111,24 @@ void ULevelStreamingPersistentPropertiesInfo::Initialize()
 		}
 	};
 
+	bool bAllPropertiesResolved = true;
 	for (const FLevelStreamingPersistentProperty& PersistentProperty : GetDefault<ULevelStreamingPersistenceSettings>()->Properties)
 	{
-		FillClassesProperties(PersistentProperty.Path, PersistentProperty.bIsPublic);
+		bAllPropertiesResolved &= FillClassesProperties(PersistentProperty.Path, PersistentProperty.bIsPublic);
 	}
 
+#if WITH_EDITOR
+	if (!bAllPropertiesResolved)
+	{
+		// Pop up warning in editor since persistence won't work for some properties. Refer to docs with syntax examples.
+		FMessageLog MessageLog("LoadErrors");
+		const FString DocsURL = "https://dev.epicgames.com/community/learning/knowledge-base/r6wl/unreal-engine-world-building-guide#wp-importantchangesin53";
+		MessageLog.Warning(LOCTEXT("InvalidPaths", "LevelStreamingPersistenceSettings in Engine.ini contained invalid property paths: see log."))
+			->AddToken(FURLToken::Create(DocsURL, LOCTEXT("DocsLink", "Documentation: World Building Guide")));
+		MessageLog.Open(EMessageSeverity::Warning);
+	}
+#endif
+
 	CreateClassDefaultPropertyBag();
 }
 
@@ -1445,3 +1461,5 @@ FAutoConsoleVariableRef ULevelStreamingPersistenceManager::EnableCommand(
 	ULevelStreamingPersistenceManager::bIsEnabled,
 	TEXT("Turn on/off to enable/disable world persistent subsystem."),
 	ECVF_Default);
+
+#undef LOCTEXT_NAMESPACE
\ No newline at end of file
