LevelStreamingPersistence: Exposed configurable settings to editor as DeveloperSettings

#Synced-CL 52723097

diff --git a/Engine/Plugins/Runtime/LevelStreamingPersistence/Source/LevelStreamingPersistence/LevelStreamingPersistence.Build.cs b/Engine/Plugins/Runtime/LevelStreamingPersistence/Source/LevelStreamingPersistence/LevelStreamingPersistence.Build.cs
--- a/Engine/Plugins/Runtime/LevelStreamingPersistence/Source/LevelStreamingPersistence/LevelStreamingPersistence.Build.cs
+++ b/Engine/Plugins/Runtime/LevelStreamingPersistence/Source/LevelStreamingPersistence/LevelStreamingPersistence.Build.cs
@@ -18,7 +18,8 @@ namespace UnrealBuildTool.Rules
 
 			PrivateDependencyModuleNames.AddRange(
 				new string[] {
-					"PropertyPath"
+					"PropertyPath",
+					"DeveloperSettings"
 				}
 			);
 		}
diff --git a/Engine/Plugins/Runtime/LevelStreamingPersistence/Source/LevelStreamingPersistence/Private/LevelStreamingPersistenceManager.cpp b/Engine/Plugins/Runtime/LevelStreamingPersistence/Source/LevelStreamingPersistence/Private/LevelStreamingPersistenceManager.cpp
--- a/Engine/Plugins/Runtime/LevelStreamingPersistence/Source/LevelStreamingPersistence/Private/LevelStreamingPersistenceManager.cpp
+++ b/Engine/Plugins/Runtime/LevelStreamingPersistence/Source/LevelStreamingPersistence/Private/LevelStreamingPersistenceManager.cpp
@@ -1112,7 +1112,8 @@ void ULevelStreamingPersistentPropertiesInfo::Initialize()
 	};
 
 	bool bAllPropertiesResolved = true;
-	for (const FLevelStreamingPersistentProperty& PersistentProperty : GetDefault<ULevelStreamingPersistenceSettings>()->Properties)
+	const TArray<FLevelStreamingPersistentProperty>& AllPersistedProperties = GetDefault<ULevelStreamingPersistenceSettings>()->GetPersistedProperties();
+	for (const FLevelStreamingPersistentProperty& PersistentProperty : AllPersistedProperties)
 	{
 		bAllPropertiesResolved &= FillClassesProperties(PersistentProperty.Path, PersistentProperty.bIsPublic);
 	}
diff --git a/Engine/Plugins/Runtime/LevelStreamingPersistence/Source/LevelStreamingPersistence/Private/LevelStreamingPersistenceSettings.cpp b/Engine/Plugins/Runtime/LevelStreamingPersistence/Source/LevelStreamingPersistence/Private/LevelStreamingPersistenceSettings.cpp
new file mode 100644
--- /dev/null
+++ b/Engine/Plugins/Runtime/LevelStreamingPersistence/Source/LevelStreamingPersistence/Private/LevelStreamingPersistenceSettings.cpp
@@ -0,0 +1,61 @@
+// Copyright Epic Games, Inc. All Rights Reserved.
+
+#include "LevelStreamingPersistenceSettings.h"
+#include "UObject/UnrealType.h"
+
+TArray<FString> ULevelStreamingPersistenceSettings::ListClassProperties()
+{
+	TSet<FString> Options;
+	for (const FLevelStreamingPersistenceClassSettings& EditableProp : PersistenceSettingsPerClass)
+	{
+		if (UClass* Class = EditableProp.ObjectClass.TryLoadClass<UObject>())
+		{
+			for (TFieldIterator<FProperty> PropertyIt(Class); PropertyIt; ++PropertyIt)
+			{
+				// Properties that aren't instance editable aren't supported for persistence.
+				// For BP variables, users can just mark them instance editable. For native
+				// variables, they should be EditAnywhere, EditInstanceOnly, etc.
+				if (!PropertyIt->HasAnyPropertyFlags(CPF_DisableEditOnInstance))
+				{
+					Options.Add(PropertyIt->GetPathName());
+				}
+			}
+		}
+	}
+
+	TArray<FString> OptionsSorted = Options.Array();
+	OptionsSorted.Sort();
+	return OptionsSorted;
+}
+
+TArray<FLevelStreamingPersistentProperty> ULevelStreamingPersistenceSettings::GetPersistedProperties() const
+{
+	// Convert into TSets to avoid duplicate entries
+	TSet<FLevelStreamingPersistentProperty> Result;
+
+	// First, append directly configured property paths. This was the original, INI-only way to configure properties to persist.
+	Result.Append(Properties);
+
+	// Then, iterate the project settings exposed class settings and gather properties from there.
+	for (const FLevelStreamingPersistenceClassSettings& EditableProp : PersistenceSettingsPerClass)
+	{
+		for (const FString& PublicProp : EditableProp.PublicPersistedProps)
+		{
+			FLevelStreamingPersistentProperty RawProp;
+			RawProp.Path = PublicProp;
+			RawProp.bIsPublic = true;
+			Result.Add(RawProp);
+		}
+
+		for (const FString& PrivateProp : EditableProp.PrivatePersistedProps)
+		{
+			FLevelStreamingPersistentProperty RawProp;
+			RawProp.Path = PrivateProp;
+			RawProp.bIsPublic = false;
+			Result.Add(RawProp);
+		}
+
+	}
+
+	return Result.Array();
+}
\ No newline at end of file
diff --git a/Engine/Plugins/Runtime/LevelStreamingPersistence/Source/LevelStreamingPersistence/Public/LevelStreamingPersistenceSettings.h b/Engine/Plugins/Runtime/LevelStreamingPersistence/Source/LevelStreamingPersistence/Public/LevelStreamingPersistenceSettings.h
--- a/Engine/Plugins/Runtime/LevelStreamingPersistence/Source/LevelStreamingPersistence/Public/LevelStreamingPersistenceSettings.h
+++ b/Engine/Plugins/Runtime/LevelStreamingPersistence/Source/LevelStreamingPersistence/Public/LevelStreamingPersistenceSettings.h
@@ -3,6 +3,7 @@
 #pragma once
 
 #include "CoreMinimal.h"
+#include "Engine/DeveloperSettings.h"
 #include "LevelStreamingPersistenceSettings.generated.h"
 
 USTRUCT()
@@ -15,15 +16,58 @@ struct LEVELSTREAMINGPERSISTENCE_API FLevelStreamingPersistentProperty
 
 	UPROPERTY()
 	bool bIsPublic = false;
+
+	friend bool operator==(const FLevelStreamingPersistentProperty& A, const FLevelStreamingPersistentProperty& B)
+	{
+		return A.Path == B.Path && A.bIsPublic == B.bIsPublic;
+	}
+
+	friend uint32 GetTypeHash(const FLevelStreamingPersistentProperty& Property)
+	{
+		return HashCombine(GetTypeHash(Property.Path), GetTypeHash(Property.bIsPublic));
+	}
 };
 
-UCLASS(Config = Engine)
-class LEVELSTREAMINGPERSISTENCE_API ULevelStreamingPersistenceSettings : public UObject
+// Editor-exposed persistence configuration per class.
+USTRUCT()
+struct LEVELSTREAMINGPERSISTENCE_API FLevelStreamingPersistenceClassSettings
 {
 	GENERATED_BODY()
 
+	// The base object class to persist properties for. Any map placed actor or component of this class or derived, will have the properties configured below persisted.
+	// Actor and component classes can be selected here. Although other cooked subobjects can be persisted too, their path name stability is an advanced topic.
+	// To configure non-actor and components, see FLevelStreamingPersistentProperty for direct use in INI.
+	UPROPERTY(EditAnywhere, Category = Persistence, meta=(AllowedClasses="/Script/Engine.Actor, /Script/Engine.ActorComponent"))
+	FSoftClassPath ObjectClass;
+	// Name of the class's property. The property must be a direct member of the class and can't be nested. 
+	// The property should be instance editable to support persistence.
+	UPROPERTY(EditAnywhere, Category = Persistence, meta = (GetOptions = "ListClassProperties"))
+	TArray<FString> PublicPersistedProps;
+	// Name of the class's property. The property must be a direct member of the class and can't be nested.
+	// The property should be instance editable to support persistence.
+	UPROPERTY(EditAnywhere, Category = Persistence, meta = (GetOptions = "ListClassProperties"))
+	TArray<FString> PrivatePersistedProps;
+};
+
+UCLASS(Config = Engine, DefaultConfig, meta=(DisplayName="Level Streaming Persistence"))
+class LEVELSTREAMINGPERSISTENCE_API ULevelStreamingPersistenceSettings : public UDeveloperSettings
+{
+	GENERATED_BODY()
+
+	// Populate property dropdown for entries in 'PersistenceSettingsPerClass'. Known limitation: currently it lists properties 
+	// for all classes selected in any PersistenceSettingsPerClass entry. Currently, the editor doesn't allow individual GetOptions 
+	// per struct entry in an array so it's not possible to resolve which struct we're editing.
+	UFUNCTION()
+	TArray<FString> ListClassProperties();
+
+	// Combines Properties and PersistenceSettingsPerClass.
+	TArray<FLevelStreamingPersistentProperty> GetPersistedProperties() const;
+	// Property paths that have been INI-configurable since plugin release
 	UPROPERTY(Config)
 	TArray<FLevelStreamingPersistentProperty> Properties;
+	// Property paths that are editor configurable, added more recently
+	UPROPERTY(Config, EditAnywhere, Category = Persistence)
+	TArray<FLevelStreamingPersistenceClassSettings> PersistenceSettingsPerClass;
 
 	friend class ULevelStreamingPersistentPropertiesInfo;
 };
\ No newline at end of file
